Refactored DerivedData StatusBar, Zen Dashboard Tool and added Zen StatusBar

Made a new Zen StatusBar that will replace the existing DerivedData StatusBar if Zen Server is available ( ie. if you don't use Zen you get all the existing UX ).
Removed the Launch Zen Dashboard option from the DerivedData StatusBar ( it makes no sense if Zen is not running ). It is now only present only in the Zen StatusBar
Added a new library ZenWidgets that contains the main Widgets from Zen Dashboard Tool so we can re-use these for the Editor and the Standalone Zen Dashboard Tool
Added a new library ZenEditor which contains the Zen StatusBar and the necessary bits and pieces to run and inject that into the Editor
Removed Widgets from the Zen Dashboard tool and put them in the shared ZenWidgets library so they can also be used by the Zen StatusBar
Added a new Zen Store widget that contains all the information in the Zen Dashbaord Tool but can be "docked and locked" in Editor sessions
Moved the Zen Launch Dashboard Icon from the Derived Data section and added a new Icon section explicitly for Zen so we can change the icons without breaking the existing DerivedData Status Bar
Added Tools, Store and Cache menu sections to the Zen StatusBar

diff --git a/Engine/Source/Editor/DerivedDataEditor/DerivedDataEditor.Build.cs b/Engine/Source/Editor/DerivedDataEditor/DerivedDataEditor.Build.cs
--- a/Engine/Source/Editor/DerivedDataEditor/DerivedDataEditor.Build.cs
+++ b/Engine/Source/Editor/DerivedDataEditor/DerivedDataEditor.Build.cs
@@ -1,7 +1,7 @@
+
 // Copyright Epic Games, Inc. All Rights Reserved.
 
 using UnrealBuildTool;
-
 public class DerivedDataEditor : ModuleRules
 {
 	public DerivedDataEditor(ReadOnlyTargetRules Target)
@@ -24,8 +24,7 @@ public class DerivedDataEditor : ModuleRules
 				"WorkspaceMenuStructure",
 				"MessageLog",
 				"ToolWidgets",
-				"UATHelper",
-				"Zen"
+				"DerivedDataWidgets"
 			});
 
 		DynamicallyLoadedModuleNames.AddRange(
@@ -39,4 +38,4 @@ public class DerivedDataEditor : ModuleRules
 			});
 		CppCompileWarningSettings.UnsafeTypeCastWarningLevel = WarningLevel.Error;
 	}
-}
+}
\ No newline at end of file
diff --git a/Engine/Source/Editor/DerivedDataEditor/Private/DerivedDataCacheNotifications.cpp b/Engine/Source/Editor/DerivedDataEditor/Private/DerivedDataCacheNotifications.cpp
deleted file mode 100644
diff --git a/Engine/Source/Editor/DerivedDataEditor/Private/DerivedDataCacheNotifications.h b/Engine/Source/Editor/DerivedDataEditor/Private/DerivedDataCacheNotifications.h
deleted file mode 100644
diff --git a/Engine/Source/Editor/DerivedDataEditor/Private/DerivedDataEditor.cpp b/Engine/Source/Editor/DerivedDataEditor/Private/DerivedDataEditor.cpp
new file mode 100644
--- /dev/null
+++ b/Engine/Source/Editor/DerivedDataEditor/Private/DerivedDataEditor.cpp
@@ -0,0 +1,123 @@
+// Copyright Epic Games, Inc. All Rights Reserved.
+
+#include "DerivedDataEditor.h"
+#include "DerivedDataCacheNotifications.h"
+#include "IDerivedDataCacheNotifications.h"
+#include "DerivedDataInformation.h"
+#include "Framework/Application/SlateApplication.h"
+#include "Modules/ModuleManager.h"
+#include "SDerivedDataDialogs.h"
+#include "SDerivedDataStatusBar.h"
+#include "Styling/AppStyle.h"
+#include "Widgets/DeclarativeSyntaxSupport.h"
+#include "Widgets/Docking/SDockTab.h"
+#include "WorkspaceMenuStructure.h"
+#include "WorkspaceMenuStructureModule.h"
+
+#define LOCTEXT_NAMESPACE "DerivedDataEditor"
+
+IMPLEMENT_MODULE(FDerivedDataEditor, DerivedDataEditor );
+
+static const FName DerivedDataResourceUsageTabName = FName(TEXT("DerivedDataResourceUsage"));
+static const FName DerivedDataCacheStatisticsTabName = FName(TEXT("DerivedDataCacheStatistics"));
+
+void FDerivedDataEditor::StartupModule()
+{
+	const FSlateIcon ResourceUsageIcon(FAppStyle::GetAppStyleSetName(), "DerivedData.ResourceUsage");
+
+	FGlobalTabmanager::Get()->RegisterNomadTabSpawner(DerivedDataResourceUsageTabName, FOnSpawnTab::CreateRaw(this, &FDerivedDataEditor::CreateResourceUsageTab))
+		.SetDisplayName(LOCTEXT("DerivedDataResourceUsageTabTitle", "Resource Usage"))
+		.SetTooltipText(LOCTEXT("DerivedDataResourceUsageTabToolTipText", "Derived Data Resource Usage"))
+		.SetGroup(WorkspaceMenu::GetMenuStructure().GetToolsCategory())
+		.SetIcon(ResourceUsageIcon);
+
+	const FSlateIcon CacheStatisticsIcon(FAppStyle::GetAppStyleSetName(), "DerivedData.Cache.Statistics");
+
+	FGlobalTabmanager::Get()->RegisterNomadTabSpawner(DerivedDataCacheStatisticsTabName, FOnSpawnTab::CreateRaw(this, &FDerivedDataEditor::CreateCacheStatisticsTab))
+		.SetDisplayName(LOCTEXT("DerivedDataCacheStatisticsTabTitle", "Cache Statistics"))
+		.SetTooltipText(LOCTEXT("DerivedDataCacheStatisticsTabToolTipText", "Derived Data Cache Statistics"))
+		.SetGroup(WorkspaceMenu::GetMenuStructure().GetToolsCategory())
+		.SetIcon(CacheStatisticsIcon);
+
+#if WITH_RELOAD
+	// This code attempts to relaunch the tabs when you reload this module
+	if (IsReloadActive() && FSlateApplication::IsInitialized())
+	{
+		ShowCacheStatisticsTab();
+		ShowResourceUsageTab();
+	}
+#endif // WITH_RELOAD
+
+	FDerivedDataEditorMenuCommands::Register();
+
+	DerivedDataCacheNotifications.Reset(new FDerivedDataCacheNotifications);
+}
+
+void FDerivedDataEditor::ShutdownModule()
+{
+	if (FSlateApplication::IsInitialized())
+	{
+		FGlobalTabmanager::Get()->UnregisterNomadTabSpawner(DerivedDataResourceUsageTabName);
+
+		if (ResourceUsageTab.IsValid())
+		{
+			ResourceUsageTab.Pin()->RequestCloseTab();
+		}
+
+		FGlobalTabmanager::Get()->UnregisterNomadTabSpawner(DerivedDataCacheStatisticsTabName);
+
+		if (CacheStatisticsTab.IsValid())
+		{
+			CacheStatisticsTab.Pin()->RequestCloseTab();
+		}
+	}
+
+	FDerivedDataEditorMenuCommands::Unregister();
+}
+
+TSharedRef<SWidget> FDerivedDataEditor::CreateStatusBarWidget()
+{
+	return SNew(SDerivedDataStatusBarWidget);
+}
+
+TSharedPtr<SWidget> FDerivedDataEditor::CreateResourceUsageDialog()
+{
+	return SNew(SDerivedDataResourceUsageDialog);
+}
+
+TSharedRef<SDockTab> FDerivedDataEditor::CreateResourceUsageTab(const FSpawnTabArgs& Args)
+{
+	return SAssignNew(ResourceUsageTab, SDockTab)
+	.TabRole(ETabRole::NomadTab)
+	[
+		CreateResourceUsageDialog().ToSharedRef()
+	];
+}
+
+void FDerivedDataEditor::ShowResourceUsageTab()
+{
+	FGlobalTabmanager::Get()->TryInvokeTab(FTabId(DerivedDataResourceUsageTabName));
+}
+
+
+TSharedPtr<SWidget> FDerivedDataEditor::CreateCacheStatisticsDialog()
+{
+	return SNew(SDerivedDataCacheStatisticsDialog);
+}
+
+
+TSharedRef<SDockTab> FDerivedDataEditor::CreateCacheStatisticsTab(const FSpawnTabArgs& Args)
+{
+	return SAssignNew(CacheStatisticsTab, SDockTab)
+		.TabRole(ETabRole::NomadTab)
+		[
+			CreateCacheStatisticsDialog().ToSharedRef()
+		];
+}
+
+void FDerivedDataEditor::ShowCacheStatisticsTab()
+{
+	FGlobalTabmanager::Get()->TryInvokeTab(FTabId(DerivedDataCacheStatisticsTabName));
+}
+
+#undef LOCTEXT_NAMESPACE
diff --git a/Engine/Source/Editor/DerivedDataEditor/Private/DerivedDataEditorModule.cpp b/Engine/Source/Editor/DerivedDataEditor/Private/DerivedDataEditorModule.cpp
deleted file mode 100644
diff --git a/Engine/Source/Editor/DerivedDataEditor/Private/DerivedDataInformation.cpp b/Engine/Source/Editor/DerivedDataEditor/Private/DerivedDataInformation.cpp
deleted file mode 100644
diff --git a/Engine/Source/Editor/DerivedDataEditor/Private/DerivedDataInformation.h b/Engine/Source/Editor/DerivedDataEditor/Private/DerivedDataInformation.h
deleted file mode 100644
diff --git a/Engine/Source/Editor/DerivedDataEditor/Private/SDerivedDataDialogs.cpp b/Engine/Source/Editor/DerivedDataEditor/Private/SDerivedDataDialogs.cpp
deleted file mode 100644
diff --git a/Engine/Source/Editor/DerivedDataEditor/Private/SDerivedDataDialogs.h b/Engine/Source/Editor/DerivedDataEditor/Private/SDerivedDataDialogs.h
deleted file mode 100644
diff --git a/Engine/Source/Editor/DerivedDataEditor/Private/SDerivedDataStatusBar.cpp b/Engine/Source/Editor/DerivedDataEditor/Private/SDerivedDataStatusBar.cpp
--- a/Engine/Source/Editor/DerivedDataEditor/Private/SDerivedDataStatusBar.cpp
+++ b/Engine/Source/Editor/DerivedDataEditor/Private/SDerivedDataStatusBar.cpp
@@ -6,7 +6,8 @@
 #include "Containers/Array.h"
 #include "Delegates/Delegate.h"
 #include "DerivedDataCacheInterface.h"
-#include "DerivedDataEditorModule.h"
+#include "IDerivedDataCacheNotifications.h"
+#include "DerivedDataEditor.h"
 #include "DerivedDataInformation.h"
 #include "Framework/Commands/InputChord.h"
 #include "Framework/Commands/UIAction.h"
@@ -45,7 +46,6 @@
 #include "Widgets/SBoxPanel.h"
 #include "Widgets/SOverlay.h"
 #include "Widgets/Text/STextBlock.h"
-#include "ZenDashboardLauncher.h"
 
 class SWidget;
 struct FSlateBrush;
@@ -53,10 +53,10 @@ struct FSlateBrush;
 #define LOCTEXT_NAMESPACE "DerivedDataEditor"
 
 
-TSharedRef<FUICommandList> FDerivedDataStatusBarMenuCommands::ActionList(new FUICommandList());
+TSharedRef<FUICommandList> FDerivedDataEditorMenuCommands::ActionList(new FUICommandList());
 
-FDerivedDataStatusBarMenuCommands::FDerivedDataStatusBarMenuCommands()
-	: TCommands<FDerivedDataStatusBarMenuCommands>
+FDerivedDataEditorMenuCommands::FDerivedDataEditorMenuCommands()
+	: TCommands<FDerivedDataEditorMenuCommands>
 	(
 		"DerivedDataSettings",
 		NSLOCTEXT("Contexts", "Derived Data", "Derived Data"),
@@ -65,7 +65,7 @@ FDerivedDataStatusBarMenuCommands::FDerivedDataStatusBarMenuCommands()
 		)
 {}
 
-void FDerivedDataStatusBarMenuCommands::RegisterCommands()
+void FDerivedDataEditorMenuCommands::RegisterCommands()
 {
 	UI_COMMAND(ChangeSettings, "Change Cache Settings", "Opens a dialog to change Cache settings.", EUserInterfaceActionType::Button, FInputChord());
 	UI_COMMAND(ViewCacheStatistics, "View Cache Statistics", "Opens the Cache Statistics panel.", EUserInterfaceActionType::Button, FInputChord());
@@ -74,46 +74,35 @@ void FDerivedDataStatusBarMenuCommands::RegisterCommands()
 
 	ActionList->MapAction(
 		ChangeSettings,
-		FExecuteAction::CreateStatic(&FDerivedDataStatusBarMenuCommands::ChangeSettings_Clicked)
+		FExecuteAction::CreateStatic(&FDerivedDataEditorMenuCommands::ChangeSettings_Clicked)
 	);
 
 	ActionList->MapAction(
 		ViewCacheStatistics,
-		FExecuteAction::CreateStatic(&FDerivedDataStatusBarMenuCommands::ViewCacheStatistics_Clicked)
+		FExecuteAction::CreateStatic(&FDerivedDataEditorMenuCommands::ViewCacheStatistics_Clicked)
 	);
 
 	ActionList->MapAction(
 		ViewResourceUsage,
-		FExecuteAction::CreateStatic(&FDerivedDataStatusBarMenuCommands::ViewResourceUsage_Clicked)
-	);
-
-	ActionList->MapAction(
-		LaunchZenDashboard,
-		FExecuteAction::CreateStatic(&FDerivedDataStatusBarMenuCommands::LaunchZenDashboard_Clicked)
+		FExecuteAction::CreateStatic(&FDerivedDataEditorMenuCommands::ViewResourceUsage_Clicked)
 	);
 }
 
-void FDerivedDataStatusBarMenuCommands::ChangeSettings_Clicked()
+void FDerivedDataEditorMenuCommands::ChangeSettings_Clicked()
 {
 	FModuleManager::LoadModuleChecked<ISettingsModule>("Settings").ShowViewer("Editor", "General", "Global");
 }
 
-void FDerivedDataStatusBarMenuCommands::ViewCacheStatistics_Clicked()
-{
-	FDerivedDataEditorModule& DerivedDataEditorModule = FModuleManager::LoadModuleChecked<FDerivedDataEditorModule>("DerivedDataEditor");
-	DerivedDataEditorModule.ShowCacheStatisticsTab();
-}
-
-void FDerivedDataStatusBarMenuCommands::ViewResourceUsage_Clicked()
+void FDerivedDataEditorMenuCommands::ViewCacheStatistics_Clicked()
 {
-	FDerivedDataEditorModule& DerivedDataEditorModule = FModuleManager::LoadModuleChecked<FDerivedDataEditorModule>("DerivedDataEditor");
-	DerivedDataEditorModule.ShowResourceUsageTab();
+	FDerivedDataEditor& DerivedDataStatusBar = FModuleManager::LoadModuleChecked<FDerivedDataEditor>("DerivedDataEditor");
+	DerivedDataStatusBar.ShowCacheStatisticsTab();
 }
 
-void FDerivedDataStatusBarMenuCommands::LaunchZenDashboard_Clicked()
+void FDerivedDataEditorMenuCommands::ViewResourceUsage_Clicked()
 {
-	using namespace UE::Zen;
-	FZenDashboardLauncher::Get()->StartZenDashboard(FZenDashboardLauncher::Get()->GetZenDashboardApplicationPath());
+	FDerivedDataEditor& DerivedDataStatusBar = FModuleManager::LoadModuleChecked<FDerivedDataEditor>("DerivedDataEditor");
+	DerivedDataStatusBar.ShowResourceUsageTab();
 }
 
 TSharedRef<SWidget> SDerivedDataStatusBarWidget::CreateStatusBarMenu()
@@ -124,38 +113,32 @@ TSharedRef<SWidget> SDerivedDataStatusBarWidget::CreateStatusBarMenu()
 		FToolMenuSection& Section = Menu->AddSection("DDCMenuSettingsSection", LOCTEXT("DDCMenuSettingsSection", "Settings"));
 
 		Section.AddMenuEntry(
-			FDerivedDataStatusBarMenuCommands::Get().ChangeSettings,
+			FDerivedDataEditorMenuCommands::Get().ChangeSettings,
 			TAttribute<FText>(),
 			TAttribute<FText>(),
 			FSlateIcon(FAppStyle::GetAppStyleSetName(), "DerivedData.Cache.Settings")
 		);
-		Section.AddMenuEntry(
-			FDerivedDataStatusBarMenuCommands::Get().LaunchZenDashboard,
-			TAttribute<FText>(),
-			TAttribute<FText>(),
-			FSlateIcon(FAppStyle::GetAppStyleSetName(), "DerivedData.LaunchZenDashboard")
-		);
 	}
 
 	{
 		FToolMenuSection& Section = Menu->AddSection("DDCMenuStatisticsSection", LOCTEXT("DDCMenuStatisticsSection", "Statistics"));
 
 		Section.AddMenuEntry(
-			FDerivedDataStatusBarMenuCommands::Get().ViewCacheStatistics,
+			FDerivedDataEditorMenuCommands::Get().ViewCacheStatistics,
 			TAttribute<FText>(),
 			TAttribute<FText>(),
 			FSlateIcon(FAppStyle::GetAppStyleSetName(), "DerivedData.Cache.Statistics")
 		);
 
 		Section.AddMenuEntry(
-			FDerivedDataStatusBarMenuCommands::Get().ViewResourceUsage,
+			FDerivedDataEditorMenuCommands::Get().ViewResourceUsage,
 			TAttribute<FText>(),
 			TAttribute<FText>(),
 			FSlateIcon(FAppStyle::GetAppStyleSetName(), "DerivedData.ResourceUsage")
 		);
 	}
 
-	return UToolMenus::Get()->GenerateWidget("StatusBar.ToolBar.DDC", FToolMenuContext(FDerivedDataStatusBarMenuCommands::ActionList));
+	return UToolMenus::Get()->GenerateWidget("StatusBar.ToolBar.DDC", FToolMenuContext(FDerivedDataEditorMenuCommands::ActionList));
 }
 
 void SDerivedDataStatusBarWidget::Construct(const FArguments& InArgs)
diff --git a/Engine/Source/Editor/DerivedDataEditor/Private/SDerivedDataStatusBar.h b/Engine/Source/Editor/DerivedDataEditor/Private/SDerivedDataStatusBar.h
--- a/Engine/Source/Editor/DerivedDataEditor/Private/SDerivedDataStatusBar.h
+++ b/Engine/Source/Editor/DerivedDataEditor/Private/SDerivedDataStatusBar.h
@@ -15,11 +15,11 @@ class SNotificationItem;
 class SWidget;
 struct FSlateBrush;
 
-class FDerivedDataStatusBarMenuCommands : public TCommands<FDerivedDataStatusBarMenuCommands>
+class FDerivedDataEditorMenuCommands : public TCommands<FDerivedDataEditorMenuCommands>
 {
 public:
 
-	FDerivedDataStatusBarMenuCommands();
+	FDerivedDataEditorMenuCommands();
 
 	virtual void RegisterCommands() override;
 
@@ -28,8 +28,7 @@ private:
 	static void ChangeSettings_Clicked();
 	static void ViewCacheStatistics_Clicked();
 	static void ViewResourceUsage_Clicked();
-	static void LaunchZenDashboard_Clicked();
-
+	
 public:
 
 	TSharedPtr< FUICommandInfo > ChangeSettings;
diff --git a/Engine/Source/Editor/DerivedDataEditor/Private/ZenDashboardLauncher.cpp b/Engine/Source/Editor/DerivedDataEditor/Private/ZenDashboardLauncher.cpp
deleted file mode 100644
diff --git a/Engine/Source/Editor/DerivedDataEditor/Private/ZenDashboardLauncher.h b/Engine/Source/Editor/DerivedDataEditor/Private/ZenDashboardLauncher.h
deleted file mode 100644
diff --git a/Engine/Source/Editor/DerivedDataEditor/Public/DerivedDataEditor.h b/Engine/Source/Editor/DerivedDataEditor/Public/DerivedDataEditor.h
new file mode 100644
--- /dev/null
+++ b/Engine/Source/Editor/DerivedDataEditor/Public/DerivedDataEditor.h
@@ -0,0 +1,57 @@
+// Copyright Epic Games, Inc. All Rights Reserved.
+
+#pragma once
+
+#include "CoreMinimal.h"
+#include "Modules/ModuleInterface.h"
+#include "Modules/ModuleManager.h"
+#include "Templates/SharedPointer.h"
+#include "Templates/UniquePtr.h"
+#include "IDerivedDataCacheNotifications.h"
+
+class FSpawnTabArgs;
+class SDerivedDataCacheSettingsDialog;
+class IDerivedDataCacheNotifications;
+class SDockTab;
+class SWidget;
+class SWindow;
+
+/**
+ * The module holding all of the UI related pieces for DerivedData
+ */
+class DERIVEDDATAEDITOR_API FDerivedDataEditor : public IModuleInterface
+{
+public:
+
+	/**
+	 * Called right after the module DLL has been loaded and the module object has been created
+	 */
+	virtual void StartupModule() override;
+
+	/**
+	 * Called before the module is unloaded, right before the module object is destroyed.
+	 */
+	virtual void ShutdownModule() override;
+
+	TSharedRef<SWidget>	CreateStatusBarWidget();
+
+	void ShowResourceUsageTab();
+	void ShowCacheStatisticsTab();
+
+private:
+
+	TSharedPtr<SWidget> CreateResourceUsageDialog();
+	TSharedPtr<SWidget> CreateCacheStatisticsDialog();
+
+	TSharedRef<SDockTab> CreateResourceUsageTab(const FSpawnTabArgs& Args);
+	TSharedRef<SDockTab> CreateCacheStatisticsTab(const FSpawnTabArgs& Args);
+
+	TWeakPtr<SDockTab> ResourceUsageTab;
+	TWeakPtr<SDockTab> CacheStatisticsTab;
+
+	TSharedPtr<SWindow>	SettingsWindow;
+	TSharedPtr<SDerivedDataCacheSettingsDialog> SettingsDialog;
+	TUniquePtr<IDerivedDataCacheNotifications>	DerivedDataCacheNotifications;
+};
+
+
diff --git a/Engine/Source/Editor/DerivedDataEditor/Public/DerivedDataEditorModule.h b/Engine/Source/Editor/DerivedDataEditor/Public/DerivedDataEditorModule.h
deleted file mode 100644
diff --git a/Engine/Source/Editor/DerivedDataEditor/Public/IDerivedDataCacheNotifications.h b/Engine/Source/Editor/DerivedDataEditor/Public/IDerivedDataCacheNotifications.h
deleted file mode 100644
diff --git a/Engine/Source/Editor/DerivedDataWidgets/DerivedDataWidgets.Build.cs b/Engine/Source/Editor/DerivedDataWidgets/DerivedDataWidgets.Build.cs
new file mode 100644
--- /dev/null
+++ b/Engine/Source/Editor/DerivedDataWidgets/DerivedDataWidgets.Build.cs
@@ -0,0 +1,42 @@
+// Copyright Epic Games, Inc. All Rights Reserved.
+
+using UnrealBuildTool;
+
+public class DerivedDataWidgets : ModuleRules
+{
+	public DerivedDataWidgets(ReadOnlyTargetRules Target)
+		 : base(Target)
+	{
+		PrivateDependencyModuleNames.AddRange(
+			new string[] {
+				"Core",
+				"CoreUObject",
+				"Engine",
+				"Slate",
+				"SlateCore",
+				"InputCore",
+				"EditorFramework",
+				"UnrealEd",
+				"ToolMenus",
+				"OutputLog",
+				"DerivedDataCache",
+				"EditorSubsystem",
+				"WorkspaceMenuStructure",
+				"MessageLog",
+				"ToolWidgets"
+			});
+
+		DynamicallyLoadedModuleNames.AddRange(
+			new string[] {
+				"MainFrame",
+			});
+
+
+		PrivateIncludePathModuleNames.AddRange(
+			new string[] {
+			});
+		CppCompileWarningSettings.UnsafeTypeCastWarningLevel = WarningLevel.Error;
+	}
+}
+
+
diff --git a/Engine/Source/Editor/DerivedDataWidgets/Private/DerivedDataCacheNotifications.cpp b/Engine/Source/Editor/DerivedDataWidgets/Private/DerivedDataCacheNotifications.cpp
new file mode 100644
--- /dev/null
+++ b/Engine/Source/Editor/DerivedDataWidgets/Private/DerivedDataCacheNotifications.cpp
@@ -0,0 +1,56 @@
+// Copyright Epic Games, Inc. All Rights Reserved.
+
+#include "DerivedDataCacheNotifications.h"
+
+#include "CoreGlobals.h"
+#include "Delegates/Delegate.h"
+#include "DerivedDataCacheInterface.h"
+#include "Logging/LogMacros.h"
+#include "Misc/CoreMisc.h"
+
+#define LOCTEXT_NAMESPACE "DerivedDataCacheNotifications"
+
+DEFINE_LOG_CATEGORY_STATIC(DerivedDataCacheNotifications, Log, All);
+
+FDerivedDataCacheNotifications::FDerivedDataCacheNotifications() :
+	bSubscribed(false)
+{
+	Subscribe(true);
+}
+
+FDerivedDataCacheNotifications::~FDerivedDataCacheNotifications()
+{
+	Subscribe(false);
+}
+
+void FDerivedDataCacheNotifications::OnDDCNotificationEvent(FDerivedDataCacheInterface::EDDCNotification DDCNotification)
+{
+	if (IsEngineExitRequested())
+	{
+		return;
+	}
+
+	// TODO : Handle any notificaiton evens here
+
+}
+
+void FDerivedDataCacheNotifications::Subscribe(bool bSubscribe)
+{
+	if (bSubscribe != bSubscribed)
+	{
+		FDerivedDataCacheInterface::FOnDDCNotification& DDCNotificationEvent = GetDerivedDataCacheRef().GetDDCNotificationEvent();
+
+		if (bSubscribe)
+		{
+			DDCNotificationEvent.AddRaw(this, &FDerivedDataCacheNotifications::OnDDCNotificationEvent);
+		}
+		else
+		{
+			DDCNotificationEvent.RemoveAll(this);
+		}
+
+		bSubscribed = bSubscribe;
+	}
+}
+
+#undef LOCTEXT_NAMESPACE
diff --git a/Engine/Source/Editor/DerivedDataWidgets/Private/DerivedDataInformation.cpp b/Engine/Source/Editor/DerivedDataWidgets/Private/DerivedDataInformation.cpp
new file mode 100644
--- /dev/null
+++ b/Engine/Source/Editor/DerivedDataWidgets/Private/DerivedDataInformation.cpp
@@ -0,0 +1,238 @@
+// Copyright Epic Games, Inc. All Rights Reserved.
+
+#include "DerivedDataInformation.h"
+#include "DerivedDataCacheUsageStats.h"
+#include "DerivedDataCacheInterface.h"
+#include "Settings/EditorProjectSettings.h"
+#include "Settings/EditorSettings.h"
+#include "ToolMenus.h"
+
+#define LOCTEXT_NAMESPACE "DerivedDataEditor"
+
+double FDerivedDataInformation::LastGetTime = 0;
+double FDerivedDataInformation::LastPutTime = 0;
+bool FDerivedDataInformation::bIsDownloading = false;
+bool FDerivedDataInformation::bIsUploading = false;
+FText FDerivedDataInformation::RemoteCacheWarningMessage;
+ERemoteCacheState FDerivedDataInformation::RemoteCacheState= ERemoteCacheState::Unavailable;
+
+static TArray<TSharedRef<const FDerivedDataCacheStatsNode>> GetCacheUsageStats()
+{
+	PRAGMA_DISABLE_DEPRECATION_WARNINGS;
+	TSharedRef<FDerivedDataCacheStatsNode> RootUsage = GetDerivedDataCacheRef().GatherUsageStats();
+	PRAGMA_ENABLE_DEPRECATION_WARNINGS;
+	TArray<TSharedRef<const FDerivedDataCacheStatsNode>> LeafUsageStats;
+	RootUsage->ForEachDescendant([&LeafUsageStats](TSharedRef<const FDerivedDataCacheStatsNode> Node)
+	{
+		if (Node->Children.IsEmpty())
+		{
+			LeafUsageStats.Add(Node);
+		}
+	});
+	return LeafUsageStats;
+}
+
+double FDerivedDataInformation::GetCacheActivitySizeBytes(bool bGet, bool bLocal)
+{
+	int64 TotalBytes = 0;
+
+#if ENABLE_COOK_STATS
+	for (const TSharedRef<const FDerivedDataCacheStatsNode>& Usage : GetCacheUsageStats())
+	{
+		if (Usage->IsLocal() != bLocal)
+		{
+			continue;
+		}
+
+		for (const auto& KVP : Usage->UsageStats)
+		{
+			const FDerivedDataCacheUsageStats& Stats = KVP.Value;
+
+			if (bGet)
+			{
+				TotalBytes += Stats.GetStats.GetAccumulatedValueAnyThread(FCookStats::CallStats::EHitOrMiss::Hit, FCookStats::CallStats::EStatType::Bytes);
+			}
+			else
+			{
+				TotalBytes += Stats.PutStats.GetAccumulatedValueAnyThread(FCookStats::CallStats::EHitOrMiss::Hit, FCookStats::CallStats::EStatType::Bytes);
+			}
+		}
+	}
+#endif // ENABLE_COOK_STATS
+
+	return (double)TotalBytes;
+}
+
+
+double FDerivedDataInformation::GetCacheActivityTimeSeconds(bool bGet, bool bLocal)
+{
+	int64 TotalCycles = 0;
+
+#if ENABLE_COOK_STATS
+	for (const TSharedRef<const FDerivedDataCacheStatsNode>& Usage : GetCacheUsageStats())
+	{
+		if (Usage->IsLocal() != bLocal)
+		{
+			continue;
+		}
+
+		for (const auto& KVP : Usage->UsageStats)
+		{
+			const FDerivedDataCacheUsageStats& Stats = KVP.Value;
+
+			if (bGet)
+			{
+				TotalCycles +=
+					(Stats.GetStats.GetAccumulatedValueAnyThread(FCookStats::CallStats::EHitOrMiss::Hit, FCookStats::CallStats::EStatType::Cycles) +
+						Stats.GetStats.GetAccumulatedValueAnyThread(FCookStats::CallStats::EHitOrMiss::Miss, FCookStats::CallStats::EStatType::Cycles));
+
+				TotalCycles +=
+					(Stats.PrefetchStats.GetAccumulatedValueAnyThread(FCookStats::CallStats::EHitOrMiss::Hit, FCookStats::CallStats::EStatType::Cycles) +
+						Stats.PrefetchStats.GetAccumulatedValueAnyThread(FCookStats::CallStats::EHitOrMiss::Miss, FCookStats::CallStats::EStatType::Cycles));
+			}
+			else
+			{
+				TotalCycles +=
+					(Stats.PutStats.GetAccumulatedValueAnyThread(FCookStats::CallStats::EHitOrMiss::Hit, FCookStats::CallStats::EStatType::Cycles) +
+						Stats.PutStats.GetAccumulatedValueAnyThread(FCookStats::CallStats::EHitOrMiss::Miss, FCookStats::CallStats::EStatType::Cycles));
+			}
+		}
+	}
+#endif // ENABLE_COOK_STATS
+
+	return (double)TotalCycles * FPlatformTime::GetSecondsPerCycle();
+}
+
+bool FDerivedDataInformation::GetHasRemoteCache()
+{
+	for (const TSharedRef<const FDerivedDataCacheStatsNode>& Usage : GetCacheUsageStats())
+	{
+		if (!Usage->IsLocal())
+		{
+				return true;
+		}
+	}
+	return false;
+}
+
+bool FDerivedDataInformation::GetHasZenCache()
+{
+	for (const TSharedRef<const FDerivedDataCacheStatsNode>& Usage : GetCacheUsageStats())
+	{
+		if (Usage->GetCacheType().Equals(TEXT("Zen")))
+		{
+				return true;
+		}
+	}
+	return false;
+}
+
+bool FDerivedDataInformation::GetHasUnrealCloudCache()
+{
+	for (const TSharedRef<const FDerivedDataCacheStatsNode>& Usage : GetCacheUsageStats())
+	{
+		if (Usage->GetCacheType().Equals(TEXT("Unreal Cloud DDC")))
+		{
+			return true;
+		}
+	}
+	return false;
+}
+
+void FDerivedDataInformation::UpdateRemoteCacheState()
+{
+	RemoteCacheState = ERemoteCacheState::Unavailable;
+
+	if ( GetHasRemoteCache() )
+	{
+		const double OldLastGetTime = LastGetTime;
+		const double OldLastPutTime = LastPutTime;
+
+		LastGetTime = FDerivedDataInformation::GetCacheActivityTimeSeconds(/*bGet*/ true, /*bLocal*/ false);
+		LastPutTime = FDerivedDataInformation::GetCacheActivityTimeSeconds(/*bGet*/ false, /*bLocal*/ false);
+
+		if (OldLastGetTime != 0.0 && OldLastPutTime != 0.0)
+		{
+			bIsDownloading = OldLastGetTime != LastGetTime;
+			bIsUploading = OldLastPutTime != LastPutTime;
+		}
+
+		if (bIsUploading || bIsDownloading)
+		{
+			RemoteCacheState = ERemoteCacheState::Busy;
+		}
+		else
+		{
+			RemoteCacheState = ERemoteCacheState::Idle;
+		}
+	}
+
+	if (const UDDCProjectSettings* DDCProjectSettings = GetDefault<UDDCProjectSettings>(); DDCProjectSettings && DDCProjectSettings->EnableWarnings)
+	{
+		const UEditorSettings* EditorSettings = GetDefault<UEditorSettings>();
+
+		if (EditorSettings && EditorSettings->bEnableDDCNotifications)
+		{
+			if (DDCProjectSettings->RecommendEveryoneUseUnrealCloudDDC && EditorSettings->bNotifyUseUnrealCloudDDC && !GetHasUnrealCloudCache())
+			{
+				RemoteCacheState = ERemoteCacheState::Warning;
+				RemoteCacheWarningMessage = FText(LOCTEXT("UnrealCloudDDCWarning", "It is recommended that you use Unreal Cloud DDC.\nDisable this notification in the settings."));
+			}
+			else if (DDCProjectSettings->RecommendEveryoneSetupAGlobalLocalDDCPath && EditorSettings->bNotifySetupDDCPath && EditorSettings->GlobalLocalDDCPath.Path.IsEmpty())
+			{
+				RemoteCacheState = ERemoteCacheState::Warning;
+				RemoteCacheWarningMessage = FText(LOCTEXT("GlobalLocalDDCPathWarning", "It is recommended that you set up a valid Global Local DDC Path.\nDisable this notification or set up a valid Global Local DDC Path in the settings."));
+			}
+			else if (DDCProjectSettings->RecommendEveryoneSetupAGlobalSharedDDCPath && EditorSettings->bNotifySetupDDCPath && EditorSettings->GlobalSharedDDCPath.Path.IsEmpty())
+			{
+				RemoteCacheState = ERemoteCacheState::Warning;
+				RemoteCacheWarningMessage = FText(LOCTEXT("GlobalSharedDDCPathWarning", "It is recommended that you set up a valid Global Shared DDC Path.\nDisable this notification or set up a valid Global Shared DDC Path in the settings."));
+			}
+			else if (DDCProjectSettings->RecommendEveryoneEnableS3DDC && EditorSettings->bNotifyEnableS3DD && !EditorSettings->bEnableS3DDC)
+			{
+				RemoteCacheState = ERemoteCacheState::Warning;
+				RemoteCacheWarningMessage = FText(LOCTEXT("AWSS3CacheEnabledWarning", "It is recommended that you enable the AWS S3 Cache.\nDisable this notification or enable the AWS S3 Cache in the settings."));
+			}
+			else if (DDCProjectSettings->RecommendEveryoneSetupAGlobalS3DDCPath && EditorSettings->bNotifySetupDDCPath && EditorSettings->GlobalS3DDCPath.Path.IsEmpty())
+			{
+				RemoteCacheState = ERemoteCacheState::Warning;
+				RemoteCacheWarningMessage = FText(LOCTEXT("S3GloblaLocalPathWarning", "It is recommended that you set up a valid Global Local S3 DDC Path.\nDisable this notification or set up a valid Global Local S3 DDC Path in the settings."));
+			}
+		}
+	}
+}
+
+
+FText FDerivedDataInformation::GetRemoteCacheStateAsText()
+{
+	switch (FDerivedDataInformation::GetRemoteCacheState())
+	{
+	case ERemoteCacheState::Idle:
+	{
+		return FText(LOCTEXT("DDCStateIdle","Idle"));
+		break;
+	}
+
+	case ERemoteCacheState::Busy:
+	{
+		return FText(LOCTEXT("DDCStateBusy", "Busy"));
+		break;
+	}
+
+	case ERemoteCacheState::Unavailable:
+	{
+		return FText(LOCTEXT("DDCStateUnavailable", "Unavailable"));
+		break;
+	}
+
+	default:
+	case ERemoteCacheState::Warning:
+	{
+		return FText(LOCTEXT("DDCStateWarning", "Warning"));
+		break;
+	}
+	}
+}
+
+
+#undef LOCTEXT_NAMESPACE
\ No newline at end of file
diff --git a/Engine/Source/Editor/DerivedDataWidgets/Private/SDerivedDataDialogs.cpp b/Engine/Source/Editor/DerivedDataWidgets/Private/SDerivedDataDialogs.cpp
new file mode 100644
--- /dev/null
+++ b/Engine/Source/Editor/DerivedDataWidgets/Private/SDerivedDataDialogs.cpp
@@ -0,0 +1,923 @@
+// Copyright Epic Games, Inc. All Rights Reserved.
+
+#include "SDerivedDataDialogs.h"
+#include "Algo/Sort.h"
+#include "DerivedDataCacheInterface.h"
+#include "DerivedDataCacheUsageStats.h"
+#include "DerivedDataInformation.h"
+#include "Framework/Application/SlateApplication.h"
+#include "Internationalization/FastDecimalFormat.h"
+#include "Math/BasicMathExpressionEvaluator.h"
+#include "Math/UnitConversion.h"
+#include "Misc/ExpressionParser.h"
+#include "Styling/StyleColors.h"
+#include "Widgets/Images/SImage.h"
+#include "Widgets/Layout/SGridPanel.h"
+#include "Widgets/SBoxPanel.h"
+#include "Widgets/SToolTip.h"
+#include "Widgets/Text/STextBlock.h"
+
+#define LOCTEXT_NAMESPACE "DerivedDataCacheEditor"
+
+template <typename ValueType>
+static FString ZeroDecimalFormat(ValueType Value)
+{
+	const FNumberFormattingOptions NumberFormattingOptions = FNumberFormattingOptions()
+		.SetUseGrouping(true)
+		.SetMinimumFractionalDigits(0)
+		.SetMaximumFractionalDigits(0);
+	return FastDecimalFormat::NumberToString(Value, ExpressionParser::GetLocalizedNumberFormattingRules(), NumberFormattingOptions);
+}
+
+FString SingleDecimalFormat(double Value)
+{
+	const FNumberFormattingOptions NumberFormattingOptions = FNumberFormattingOptions()
+		.SetUseGrouping(true)
+		.SetMinimumFractionalDigits(1)
+		.SetMaximumFractionalDigits(1);
+	return FastDecimalFormat::NumberToString(Value, ExpressionParser::GetLocalizedNumberFormattingRules(), NumberFormattingOptions);
+}
+
+void SDerivedDataRemoteStoreDialog::Construct(const FArguments& InArgs)
+{
+	this->ChildSlot
+	[
+		SNew(SVerticalBox)
+		+ SVerticalBox::Slot()
+		.AutoHeight()
+		.Padding(0, 20, 0, 0)
+		.Expose(GridSlot)
+		[
+			GetGridPanel()
+		]
+	];
+
+	RegisterActiveTimer(0.5f, FWidgetActiveTimerDelegate::CreateSP(this, &SDerivedDataRemoteStoreDialog::UpdateGridPanels));
+}
+
+EActiveTimerReturnType SDerivedDataRemoteStoreDialog::UpdateGridPanels(double InCurrentTime, float InDeltaTime)
+{
+	(*GridSlot)
+	[
+		GetGridPanel()
+	];
+
+	SlatePrepass(GetPrepassLayoutScaleMultiplier());
+
+	return EActiveTimerReturnType::Continue;
+}
+
+TSharedRef<SWidget> SDerivedDataRemoteStoreDialog::GetGridPanel()
+{
+	// Gather the latest resource stats
+	TArray<FDerivedDataCacheResourceStat> DDCResourceStats;
+	GatherDerivedDataCacheResourceStats(DDCResourceStats);
+
+	FDerivedDataCacheResourceStat DDCResourceStatsTotal(TEXT("Total"));
+
+	// Find the total
+	for (const FDerivedDataCacheResourceStat& Stat : DDCResourceStats)
+	{
+		if (Stat.AssetType == TEXT("Total"))
+		{
+			DDCResourceStatsTotal = Stat;
+			break;
+		}
+	}
+
+	const double DownloadedBytesMB = FUnitConversion::Convert(FDerivedDataInformation::GetCacheActivitySizeBytes(true, false), EUnit::Bytes, EUnit::Megabytes);
+	const double UploadedBytesMB = FUnitConversion::Convert(FDerivedDataInformation::GetCacheActivitySizeBytes(false, false), EUnit::Bytes, EUnit::Megabytes);
+
+	TSharedRef<SGridPanel> Panel =
+		SNew(SGridPanel);
+
+	int32 Row = 0;
+
+	Panel->AddSlot(0, Row)
+	[
+		SNew(STextBlock)
+		.Font(FCoreStyle::GetDefaultFontStyle("Bold", 10))
+		.ColorAndOpacity(FStyleColors::Foreground)
+		.Text(LOCTEXT("Remote Storage", "Remote Storage"))
+	];
+
+	Row++;
+
+	Panel->AddSlot(0, Row)
+	[
+		SNew(STextBlock)
+		.Text(LOCTEXT("Connected", "Connected"))
+	];
+
+	Panel->AddSlot(1, Row)
+	.HAlign(HAlign_Right)
+	[
+		SNew(STextBlock)
+		.Text(FDerivedDataInformation::GetHasRemoteCache() ? LOCTEXT("True", "True") : LOCTEXT("False", "False"))
+	];
+
+	Row++;
+
+	Panel->AddSlot(0, Row)
+	[
+		SNew(STextBlock)
+		.Text(LOCTEXT("Downloaded", "Downloaded"))
+	];
+
+	Panel->AddSlot(1, Row)
+	.HAlign(HAlign_Right)
+	[
+		SNew(STextBlock)
+		.Text(FText::FromString(SingleDecimalFormat(DownloadedBytesMB) + TEXT(" MiB")))
+	];
+
+	Row++;
+
+	Panel->AddSlot(0, Row)
+	[
+		SNew(STextBlock)
+		.Text(LOCTEXT("Uploaded", "Uploaded"))
+	];
+
+	Panel->AddSlot(1, Row)
+	.HAlign(HAlign_Right)
+	[
+		SNew(STextBlock)
+		.Text(FText::FromString(SingleDecimalFormat(UploadedBytesMB) + TEXT(" MiB")))
+	];
+
+	Row++;
+
+	return Panel;
+}
+
+
+
+void SDerivedDataResourceUsageDialog::Construct(const FArguments& InArgs)
+{
+	this->ChildSlot
+	[
+		SNew(SVerticalBox)
+		+ SVerticalBox::Slot()
+		.AutoHeight()
+		.Padding(0, 20, 0, 0)
+		.Expose(GridSlot)
+		[
+			GetGridPanel()
+		]
+	];
+
+	RegisterActiveTimer(0.5f, FWidgetActiveTimerDelegate::CreateSP(this, &SDerivedDataResourceUsageDialog::UpdateGridPanels));
+}
+
+EActiveTimerReturnType SDerivedDataResourceUsageDialog::UpdateGridPanels(double InCurrentTime, float InDeltaTime)
+{
+	(*GridSlot)
+	[
+		GetGridPanel()
+	];
+
+	SlatePrepass(GetPrepassLayoutScaleMultiplier());
+
+	return EActiveTimerReturnType::Continue;
+}
+
+TSharedRef<SWidget> SDerivedDataResourceUsageDialog::GetGridPanel()
+{
+	// Gather the latest resource stats
+	TArray<FDerivedDataCacheResourceStat> DDCResourceStats;
+	GatherDerivedDataCacheResourceStats(DDCResourceStats);
+
+	// Sort results on descending build count, then descending load size, then ascending asset type.
+	const auto CompareStats = [](const FDerivedDataCacheResourceStat& A, const FDerivedDataCacheResourceStat& B)
+	{
+		if (A.BuildCount != B.BuildCount)
+		{
+			return A.BuildCount > B.BuildCount;
+		}
+		if (A.LoadSizeMB != B.LoadSizeMB)
+		{
+			return A.LoadSizeMB > B.LoadSizeMB;
+		}
+		return A.AssetType.Compare(B.AssetType, ESearchCase::IgnoreCase) < 0;
+	};
+	Algo::Sort(DDCResourceStats, CompareStats);
+
+	FDerivedDataCacheResourceStat DDCResourceStatsTotal(TEXT("Total"));
+
+	// Find the total and store it off, we will render that out separately
+	for (const FDerivedDataCacheResourceStat& Stat : DDCResourceStats)
+	{
+		if (Stat.AssetType == TEXT("Total"))
+		{
+			DDCResourceStatsTotal = Stat;
+			break;
+		}
+	}
+
+	TSharedRef<SGridPanel> Panel =
+		SNew(SGridPanel);
+
+	int32 Row = 0;
+
+	const float RowMargin = 0.0f;
+	const float ColumnMargin = 10.0f;
+	const FMargin TitleMargin(0.0f, 10.0f, ColumnMargin, 10.0f);
+	const FMargin TitleMarginFirstColumn(ColumnMargin, 10.0f);
+	const FSlateColor TitleColor = FStyleColors::AccentWhite;
+	const FSlateFontInfo TitleFont = FCoreStyle::GetDefaultFontStyle("Bold", 10);
+	const FMargin DefaultMargin(0.0f, RowMargin, ColumnMargin, RowMargin);
+	const FMargin DefaultMarginFirstColumn(ColumnMargin, RowMargin);
+
+	Panel->AddSlot(3, Row)
+	.HAlign(HAlign_Center)
+	[
+		SNew(STextBlock)
+		.Margin(DefaultMargin)
+		.ColorAndOpacity(TitleColor)
+		.Font(TitleFont)
+		.Text(LOCTEXT("Loaded", "Loaded"))
+	];
+
+	Panel->AddSlot(6, Row)
+	.HAlign(HAlign_Center)
+	[
+		SNew(STextBlock)
+		.Margin(DefaultMargin)
+		.ColorAndOpacity(TitleColor)
+		.Font(TitleFont)
+		.Text(LOCTEXT("Built", "Built"))
+	];
+
+	Row++;
+
+	Panel->AddSlot(0, Row)
+	[
+		SNew(STextBlock)
+		.Margin(TitleMarginFirstColumn)
+		.ColorAndOpacity(TitleColor)
+		.Font(TitleFont)
+		.Text(LOCTEXT("ResourceType", "Resource Type"))
+	];
+
+	Panel->AddSlot(1, Row)
+	.HAlign(HAlign_Right)
+	[
+		SNew(STextBlock)
+		.Margin(TitleMargin)
+	.ColorAndOpacity(TitleColor)
+	.Font(TitleFont)
+	.Text(LOCTEXT("Hit", "Hit%"))
+	];
+
+	Panel->AddSlot(2, Row)
+	.HAlign(HAlign_Right)
+	[
+		SNew(STextBlock)
+		.Margin(TitleMargin)
+		.ColorAndOpacity(TitleColor)
+		.Font(TitleFont)
+		.Text(LOCTEXT("Count", "Count"))
+	];
+
+	Panel->AddSlot(3, Row)
+	.HAlign(HAlign_Right)
+	[
+		SNew(STextBlock)
+		.Margin(TitleMargin)
+		.ColorAndOpacity(TitleColor)
+		.Font(TitleFont)
+		.Text(LOCTEXT("Time (Sec)", "Time (Sec)"))
+	];
+
+	Panel->AddSlot(4, Row)
+	.HAlign(HAlign_Right)
+	[
+		SNew(STextBlock)
+		.Margin(TitleMargin)
+		.ColorAndOpacity(TitleColor)
+		.Font(TitleFont)
+		.Text(LOCTEXT("Size (MiB)", "Size (MiB)"))
+	];
+
+	Panel->AddSlot(5, Row)
+	.HAlign(HAlign_Right)
+	[
+		SNew(STextBlock)
+		.Margin(TitleMargin)
+		.ColorAndOpacity(TitleColor)
+		.Font(TitleFont)
+		.Text(LOCTEXT("Count", "Count"))
+	];
+
+	Panel->AddSlot(6, Row)
+	.HAlign(HAlign_Right)
+	[
+		SNew(STextBlock)
+		.Margin(TitleMargin)
+		.ColorAndOpacity(TitleColor)
+		.Font(TitleFont)
+		.Text(LOCTEXT("Time (Sec)", "Time (Sec)"))
+	];
+
+	Panel->AddSlot(7, Row)
+	.HAlign(HAlign_Right)
+	[
+		SNew(STextBlock)
+		.Margin(TitleMargin)
+		.ColorAndOpacity(TitleColor)
+		.Font(TitleFont)
+		.Text(LOCTEXT("Size (MiB)", "Size (MiB)"))
+	];
+
+	Row++;
+
+	for (const FDerivedDataCacheResourceStat& Stat : DDCResourceStats)
+	{	
+		if (Stat.AssetType == TEXT("Total"))
+		{
+			continue;
+		}
+
+		Panel->AddSlot(0, Row)
+		[
+			SNew(STextBlock)
+			.Margin(DefaultMarginFirstColumn)
+			.Text(FText::FromString(Stat.AssetType))
+		];
+
+		Panel->AddSlot(1, Row)
+		.HAlign(HAlign_Right)
+		[
+			SNew(STextBlock)
+			.Margin(DefaultMargin)
+			.Text(FText::FromString(SingleDecimalFormat(Stat.Efficiency*100.0)))
+		];
+
+
+		Panel->AddSlot(2, Row)
+		.HAlign(HAlign_Right)
+		[
+			SNew(STextBlock)
+			.Margin(DefaultMargin)
+			.Text(FText::FromString(ZeroDecimalFormat(Stat.LoadCount)))
+		];
+
+		Panel->AddSlot(3, Row)
+		.HAlign(HAlign_Right)
+		[
+			SNew(STextBlock)
+			.Margin(DefaultMargin)
+			.Text(FText::FromString(SingleDecimalFormat(Stat.LoadTimeSec)))
+		];
+
+		Panel->AddSlot(4, Row)
+		.HAlign(HAlign_Right)
+		[
+			SNew(STextBlock)
+			.Margin(DefaultMargin)
+			.Text(FText::FromString(SingleDecimalFormat(Stat.LoadSizeMB)))
+		];
+
+		Panel->AddSlot(5, Row)
+		.HAlign(HAlign_Right)
+		[
+			SNew(STextBlock)
+			.Margin(DefaultMargin)
+			.Text(FText::FromString(ZeroDecimalFormat(Stat.BuildCount)))
+		];
+
+		Panel->AddSlot(6, Row)
+		.HAlign(HAlign_Right)
+		[
+			SNew(STextBlock)
+			.Margin(DefaultMargin)
+			.Text(FText::FromString(SingleDecimalFormat(Stat.BuildTimeSec)))
+		];
+
+		Panel->AddSlot(7, Row)
+		.HAlign(HAlign_Right)
+		[
+			SNew(STextBlock)
+			.Margin(DefaultMargin)
+			.Text(FText::FromString(SingleDecimalFormat(Stat.BuildSizeMB)))
+		];
+
+		Row++;
+	}
+	
+	Panel->AddSlot(0, Row)
+	[
+		SNew(STextBlock)
+		.Margin(TitleMarginFirstColumn)
+		.ColorAndOpacity(TitleColor)
+		.Font(TitleFont)
+		.Text(FText::FromString(DDCResourceStatsTotal.AssetType))
+	];
+
+	Panel->AddSlot(1, Row)
+	.HAlign(HAlign_Right)
+	[
+		SNew(STextBlock)
+		.Margin(TitleMargin)
+		.ColorAndOpacity(TitleColor)
+		.Font(TitleFont)
+		.Text(FText::FromString(SingleDecimalFormat(DDCResourceStatsTotal.Efficiency * 100.0)))
+	];
+
+
+	Panel->AddSlot(2, Row)
+	.HAlign(HAlign_Right)
+	[
+		SNew(STextBlock)
+		.Margin(TitleMargin)
+		.ColorAndOpacity(TitleColor)
+		.Font(TitleFont)
+		.Text(FText::FromString(ZeroDecimalFormat(DDCResourceStatsTotal.LoadCount)))
+	];
+
+	Panel->AddSlot(3, Row)
+	.HAlign(HAlign_Right)
+	[
+		SNew(STextBlock)
+		.Margin(TitleMargin)
+		.ColorAndOpacity(TitleColor)
+		.Font(TitleFont)
+		.Text(FText::FromString(SingleDecimalFormat(DDCResourceStatsTotal.LoadTimeSec)))
+	];
+
+	Panel->AddSlot(4, Row)
+	.HAlign(HAlign_Right)
+	[
+		SNew(STextBlock)
+		.Margin(TitleMargin)
+		.ColorAndOpacity(TitleColor)
+		.Font(TitleFont)
+		.Text(FText::FromString(SingleDecimalFormat(DDCResourceStatsTotal.LoadSizeMB)))
+	];
+
+	Panel->AddSlot(5, Row)
+	.HAlign(HAlign_Right)
+	[
+		SNew(STextBlock)
+		.Margin(TitleMargin)
+		.ColorAndOpacity(TitleColor)
+		.Font(TitleFont)
+		.Text(FText::FromString(ZeroDecimalFormat(DDCResourceStatsTotal.BuildCount)))
+	];
+
+	Panel->AddSlot(6, Row)
+	.HAlign(HAlign_Right)
+	[
+		SNew(STextBlock)
+		.Margin(TitleMargin)
+		.ColorAndOpacity(TitleColor)
+		.Font(TitleFont)
+		.Text(FText::FromString(SingleDecimalFormat(DDCResourceStatsTotal.BuildTimeSec)))
+	];
+
+	Panel->AddSlot(7, Row)
+	.HAlign(HAlign_Right)
+	[
+		SNew(STextBlock)
+		.Margin(TitleMargin)
+		.ColorAndOpacity(TitleColor)
+		.Font(TitleFont)
+		.Text(FText::FromString(SingleDecimalFormat(DDCResourceStatsTotal.BuildSizeMB)))
+	];
+
+	return Panel;
+}
+
+
+
+void SDerivedDataCacheStatisticsDialog::Construct(const FArguments& InArgs)
+{
+	const float RowMargin = 0.0f;
+	const float TitleMargin = 10.0f;
+	const float ColumnMargin = 10.0f;
+	const FSlateColor TitleColor = FStyleColors::AccentWhite;
+	const FSlateFontInfo TitleFont = FCoreStyle::GetDefaultFontStyle("Bold", 10);
+
+	this->ChildSlot
+	[
+		SNew(SVerticalBox)
+		+ SVerticalBox::Slot()
+		.Padding(0, 20, 0, 0)
+		.AutoHeight()
+		[
+			SNew(SHorizontalBox)
+			+SHorizontalBox::Slot()
+			.FillWidth(1.0f)
+			[
+				SNew(STextBlock)
+				.Margin(TitleMargin)
+				.ColorAndOpacity(TitleColor)
+				.Font(TitleFont)
+				.Justification(ETextJustify::Left)
+				.Text(FText::FromString(GetDerivedDataCache()->GetGraphName()))
+			]
+		]
+		+ SVerticalBox::Slot()
+		.AutoHeight()
+		.Padding(0, 5, 0, 0)
+		.Expose(GridSlot)
+		[
+			GetGridPanel()
+		]
+	];
+
+	RegisterActiveTimer(0.5f, FWidgetActiveTimerDelegate::CreateSP(this, &SDerivedDataCacheStatisticsDialog::UpdateGridPanels));
+}
+
+EActiveTimerReturnType SDerivedDataCacheStatisticsDialog::UpdateGridPanels(double InCurrentTime, float InDeltaTime)
+{
+	(*GridSlot)
+	[
+		GetGridPanel()
+	];
+
+	SlatePrepass(GetPrepassLayoutScaleMultiplier());
+
+	return EActiveTimerReturnType::Continue;
+}
+
+
+
+TSharedRef<SWidget> SDerivedDataCacheStatisticsDialog::GetGridPanel()
+{
+	PRAGMA_DISABLE_DEPRECATION_WARNINGS;
+	TSharedRef<FDerivedDataCacheStatsNode> RootUsage = GetDerivedDataCache()->GatherUsageStats();
+	PRAGMA_ENABLE_DEPRECATION_WARNINGS;
+	TArray<TSharedRef<const FDerivedDataCacheStatsNode>> LeafUsageStats;
+	RootUsage->ForEachDescendant([&LeafUsageStats](TSharedRef<const FDerivedDataCacheStatsNode> Node)
+	{
+		if (Node->Children.IsEmpty())
+		{
+			LeafUsageStats.Add(Node);
+		}
+	});
+
+	TSharedRef<SGridPanel> Panel =
+		SNew(SGridPanel);
+
+	const float RowMargin = 0.0f;
+	const float ColumnMargin = 10.0f;
+	const FSlateColor TitleColor = FStyleColors::AccentWhite;
+	const FSlateFontInfo TitleFont = FCoreStyle::GetDefaultFontStyle("Bold", 10);
+	const FMargin DefaultMarginFirstColumn(ColumnMargin, RowMargin);
+
+#if ENABLE_COOK_STATS
+
+	int32 Row = 0;
+
+	const FMargin TitleMargin(0.0f, 10.0f, ColumnMargin, 10.0f);
+	const FMargin TitleMarginFirstColumn(ColumnMargin, 10.0f);
+	const FMargin DefaultMargin(0.0f, RowMargin, ColumnMargin, RowMargin);
+	
+	Panel->AddSlot(0, Row)
+	[
+		SNew(STextBlock)
+	];
+
+	Panel->AddSlot(1, Row)
+	[
+		SNew(STextBlock)
+		.Margin(TitleMarginFirstColumn)
+		.Font(FCoreStyle::GetDefaultFontStyle("Bold", 10))
+		.ColorAndOpacity(TitleColor)
+		.Text(LOCTEXT("CacheType", "Cache Type"))
+	];
+
+	Panel->AddSlot(2, Row)
+	[
+		SNew(STextBlock)
+		.Margin(TitleMargin)
+		.ColorAndOpacity(TitleColor)
+		.Font(TitleFont)
+		.Text(LOCTEXT("Location", "Location"))
+	];
+
+	Panel->AddSlot(3, Row)
+	.HAlign(HAlign_Right)
+	[
+		SNew(STextBlock)
+		.Margin(TitleMargin)
+		.ColorAndOpacity(TitleColor)
+		.Font(TitleFont)
+		.Text(LOCTEXT("HitPercentage", "Hit%"))
+	];
+
+	Panel->AddSlot(4, Row)
+	.HAlign(HAlign_Right)
+	[
+		SNew(STextBlock)
+		.Margin(TitleMargin)
+		.ColorAndOpacity(TitleColor)
+		.Font(TitleFont)
+		.Text(LOCTEXT("Read", "Read (MiB)"))
+		.AutoWrapText(true)
+		.WrapTextAt(66.0f)
+	];
+
+	Panel->AddSlot(5, Row)
+	.HAlign(HAlign_Right)
+	[
+		SNew(STextBlock)
+		.Margin(TitleMargin)
+		.ColorAndOpacity(TitleColor)
+		.Font(TitleFont)
+		.Text(LOCTEXT("Write", "Write (MiB)"))
+		.AutoWrapText(true)
+		.WrapTextAt(66.0f)
+	];
+
+	Panel->AddSlot(6, Row)
+	.HAlign(HAlign_Right)
+	[
+		SNew(STextBlock)
+		.Margin(TitleMargin)
+		.ColorAndOpacity(TitleColor)
+		.Font(TitleFont)
+		.Text(LOCTEXT("Latency", "Latency (ms)"))
+		.AutoWrapText(true)
+		.WrapTextAt(66.0f)
+	];
+
+	Panel->AddSlot(7, Row)
+	.HAlign(HAlign_Right)
+	[
+		SNew(STextBlock)
+		.Margin(TitleMargin)
+		.ColorAndOpacity(TitleColor)
+		.Font(TitleFont)
+		.Text(LOCTEXT("Read Speed", "Read Speed (MiB/s)"))
+		.AutoWrapText(true)
+		.WrapTextAt(100.0f)
+	];
+
+	Panel->AddSlot(8, Row)
+	.HAlign(HAlign_Right)
+	[
+		SNew(STextBlock)
+		.Margin(TitleMargin)
+		.ColorAndOpacity(TitleColor)
+		.Font(TitleFont)
+		.Text(LOCTEXT("Write Speed", "Write Speed (MiB/s)"))
+		.AutoWrapText(true)
+		.WrapTextAt(100.0f)
+	];
+
+	Panel->AddSlot(9, Row)
+	[
+		SNew(STextBlock)
+		.Margin(TitleMargin)
+		.ColorAndOpacity(TitleColor)
+		.Font(TitleFont)
+		.Text(LOCTEXT("Storage Size", "Storage Size"))
+	];
+
+	Panel->AddSlot(10, Row)
+	[
+		SNew(STextBlock)
+		.Margin(TitleMargin)
+		.ColorAndOpacity(TitleColor)
+		.Font(TitleFont)
+		.Text(LOCTEXT("Path", "Path"))
+	];
+
+	Row++;
+
+	double SumTotalGetMB = 0.0;
+	double SumTotalPutMB = 0.0;
+	uint64 SumTotalPhysicalSize = ~0ull;
+
+	for (TSharedRef<const FDerivedDataCacheStatsNode> Node : LeafUsageStats)
+	{
+		FDerivedDataCacheUsageStats Stats;
+
+		for (const auto& KVP : Node->UsageStats)
+		{
+			Stats.Combine(KVP.Value);
+		}
+
+		const int64 TotalGetBytes = Stats.GetStats.GetAccumulatedValueAnyThread(FCookStats::CallStats::EHitOrMiss::Hit, FCookStats::CallStats::EStatType::Bytes);
+		const int64 TotalPutBytes = Stats.PutStats.GetAccumulatedValueAnyThread(FCookStats::CallStats::EHitOrMiss::Hit, FCookStats::CallStats::EStatType::Bytes);
+
+		const int64 TotalGetHits = Stats.GetStats.GetAccumulatedValueAnyThread(FCookStats::CallStats::EHitOrMiss::Hit, FCookStats::CallStats::EStatType::Counter);
+		const int64 TotalGetMisses = Stats.GetStats.GetAccumulatedValueAnyThread(FCookStats::CallStats::EHitOrMiss::Miss, FCookStats::CallStats::EStatType::Counter);
+		const int64 TotalRequests = TotalGetHits + TotalGetMisses;
+		const double HitRate = TotalRequests > 0 ? 100.0 * static_cast<double>(TotalGetHits) / static_cast<double>(TotalRequests) : 0.0;
+
+		const double TotalGetMB = FUnitConversion::Convert(static_cast<double>(TotalGetBytes), EUnit::Bytes, EUnit::Megabytes);
+		const double TotalPutMB = FUnitConversion::Convert(static_cast<double>(TotalPutBytes), EUnit::Bytes, EUnit::Megabytes);
+
+		SumTotalGetMB += TotalGetMB;
+		SumTotalPutMB += TotalPutMB;
+
+		const uint64 TotalPhysicalSize = Node->GetTotalPhysicalSize();
+		if (TotalPhysicalSize != ~0ull)
+		{
+			if (SumTotalPhysicalSize == ~0ull)
+			{
+				SumTotalPhysicalSize = TotalPhysicalSize;
+			}
+			else
+			{
+				SumTotalPhysicalSize += TotalPhysicalSize;
+			}
+		}
+
+		TSharedPtr<SImage> StatusIcon;
+		switch (Node->GetCacheStatus())
+		{
+			case EDerivedDataCacheStatus::Information:
+				StatusIcon = SNew(SImage)
+					.Image(FAppStyle::GetBrush("Icons.Help"))
+					.ToolTipText(FText::FromString(Node->GetCacheStatusText()));
+			break;
+			case EDerivedDataCacheStatus::Warning:
+				StatusIcon = SNew(SImage)
+					.Image(FAppStyle::GetBrush("Icons.Warning"))
+					.ToolTipText(FText::FromString(Node->GetCacheStatusText()));
+			break;
+			case EDerivedDataCacheStatus::Error:
+				StatusIcon = SNew(SImage)
+					.Image(FAppStyle::GetBrush("Icons.Error"))
+					.ToolTipText(FText::FromString(Node->GetCacheStatusText()));
+			break;
+			case EDerivedDataCacheStatus::Deactivation:
+				StatusIcon = SNew(SImage)
+					.Image(FAppStyle::GetBrush("Icons.X"))
+					.ToolTipText(FText::FromString(Node->GetCacheStatusText()));
+			break;
+
+		}
+
+		if (StatusIcon.IsValid())
+		{
+			Panel->AddSlot(0, Row)
+			[
+				SNew(SHorizontalBox)
+					.ToolTipText(FText::FromString(Node->GetCacheStatusText()))
+					// Status icon
+					+SHorizontalBox::Slot()
+					.AutoWidth()
+					.VAlign(VAlign_Center)
+					[
+						StatusIcon.ToSharedRef()
+					]
+			];
+		}
+		else
+		{
+			Panel->AddSlot(0, Row)
+			[
+				SNew(SHorizontalBox)
+			];
+		}
+
+		Panel->AddSlot(1, Row)
+		[
+			SNew(STextBlock)
+			.Margin(DefaultMarginFirstColumn)
+			.Text(FText::FromString(Node->GetCacheType()))
+		];
+
+		Panel->AddSlot(2, Row)
+		[
+			SNew(STextBlock)
+			.Margin(DefaultMargin)
+			.Text(Node->IsLocal() ? LOCTEXT("Local", "Local") : LOCTEXT("Remote", "Remote"))
+		];
+
+		Panel->AddSlot(3, Row)
+		.HAlign(HAlign_Right)
+		[
+			SNew(STextBlock)
+			.Margin(DefaultMargin)
+			.Text(FText::FromString(SingleDecimalFormat(HitRate)))
+		];
+
+		Panel->AddSlot(4, Row)
+		.HAlign(HAlign_Right)
+		[
+			SNew(STextBlock)
+			.Margin(DefaultMargin)
+			.Text(FText::FromString(SingleDecimalFormat(TotalGetMB)))
+		];
+
+		Panel->AddSlot(5, Row)
+		.HAlign(HAlign_Right)
+		[
+			SNew(STextBlock)
+			.Margin(DefaultMargin)
+			.Text(FText::FromString(SingleDecimalFormat(TotalPutMB)))
+		];
+
+		if (Node->SpeedStats.LatencyMS != 0.0)
+		{
+			Panel->AddSlot(6, Row)
+			.HAlign(HAlign_Right)
+			[
+				SNew(STextBlock)
+				.Margin(DefaultMargin)
+				.Text(FText::FromString(SingleDecimalFormat(Node->SpeedStats.LatencyMS)))
+			];
+		}
+
+		if (Node->SpeedStats.ReadSpeedMBs != 0.0)
+		{
+			Panel->AddSlot(7, Row)
+			.HAlign(HAlign_Right)
+			[
+				SNew(STextBlock)
+				.Margin(DefaultMargin)
+				.Text(FText::FromString(SingleDecimalFormat(Node->SpeedStats.ReadSpeedMBs)))
+			];
+		}
+
+		if (Node->SpeedStats.WriteSpeedMBs != 0.0)
+		{
+			Panel->AddSlot(8, Row)
+			.HAlign(HAlign_Right)
+			[
+				SNew(STextBlock)
+				.Margin(DefaultMargin)
+				.Text(FText::FromString(SingleDecimalFormat(Node->SpeedStats.WriteSpeedMBs)))
+			];
+		}
+
+		Panel->AddSlot(9, Row)
+		[
+			SNew(STextBlock)
+			.Margin(DefaultMargin)
+			.Text(Node->GetTotalPhysicalSize() == ~0ull ? LOCTEXT("N/A", "N/A") : FText::AsMemory(Node->GetTotalPhysicalSize(), EMemoryUnitStandard::IEC))
+		];
+
+		Panel->AddSlot(10, Row)
+		[
+			SNew(STextBlock)
+			.Margin(DefaultMargin)
+			.Text(FText::FromString(Node->GetCacheName()))
+		];
+
+		Row++;
+	}
+
+	Panel->AddSlot(1, Row)
+	[
+		SNew(STextBlock)
+		.Margin(TitleMarginFirstColumn)
+		.ColorAndOpacity(TitleColor)
+		.Font(TitleFont)
+		.Text(LOCTEXT("Total", "Total"))
+	];
+
+	Panel->AddSlot(4, Row)
+	.HAlign(HAlign_Right)
+	[
+		SNew(STextBlock)
+		.Margin(TitleMargin)
+		.ColorAndOpacity(TitleColor)
+		.Font(TitleFont)
+		.Text(FText::FromString(SingleDecimalFormat(SumTotalGetMB)))
+	];
+
+	Panel->AddSlot(5, Row)
+	.HAlign(HAlign_Right)
+	[
+		SNew(STextBlock)
+		.Margin(TitleMargin)
+		.ColorAndOpacity(TitleColor)
+		.Font(TitleFont)
+		.Text(FText::FromString(SingleDecimalFormat(SumTotalPutMB)))
+	];
+
+	Panel->AddSlot(9, Row)
+	.HAlign(HAlign_Right)
+	[
+		SNew(STextBlock)
+		.Margin(TitleMargin)
+		.ColorAndOpacity(TitleColor)
+		.Font(TitleFont)
+		.Text(SumTotalPhysicalSize == ~0ull ? LOCTEXT("N/A", "N/A") : FText::AsMemory(SumTotalPhysicalSize, EMemoryUnitStandard::IEC))
+	];
+
+#else
+	Panel->AddSlot(0, 0)
+	[
+		SNew(STextBlock)
+		.Margin(DefaultMarginFirstColumn)
+		.ColorAndOpacity(TitleColor)
+		.Font(TitleFont)
+		.Justification(ETextJustify::Center)
+		.Text(LOCTEXT("Disabled", "Cooking Stats Are Disabled For This Project"))
+	];
+
+#endif // ENABLE_COOK_STATS
+
+	return Panel;
+}
+
+#undef LOCTEXT_NAMESPACE
diff --git a/Engine/Source/Editor/DerivedDataWidgets/Public/DerivedDataCacheNotifications.h b/Engine/Source/Editor/DerivedDataWidgets/Public/DerivedDataCacheNotifications.h
new file mode 100644
--- /dev/null
+++ b/Engine/Source/Editor/DerivedDataWidgets/Public/DerivedDataCacheNotifications.h
@@ -0,0 +1,31 @@
+// Copyright Epic Games, Inc. All Rights Reserved.
+
+#pragma once
+
+#include "DerivedDataCacheInterface.h"
+#include "IDerivedDataCacheNotifications.h"
+#include "Templates/SharedPointer.h"
+
+class SNotificationItem;
+
+class DERIVEDDATAWIDGETS_API FDerivedDataCacheNotifications : public IDerivedDataCacheNotifications
+{
+public:
+	FDerivedDataCacheNotifications();
+	virtual ~FDerivedDataCacheNotifications();
+
+private: 
+
+	/** DDC data put notification handler */
+	void OnDDCNotificationEvent(FDerivedDataCacheInterface::EDDCNotification DDCNotification);
+	
+	/** Subscribe to the notifications */
+	void Subscribe(bool bSubscribe);
+
+	/** Whether we are subscribed or not **/
+	bool bSubscribed;
+
+	/** Valid when a DDC notification item is being presented */
+	TSharedPtr<SNotificationItem> SharedDDCNotification;
+};
+
diff --git a/Engine/Source/Editor/DerivedDataWidgets/Public/DerivedDataInformation.h b/Engine/Source/Editor/DerivedDataWidgets/Public/DerivedDataInformation.h
new file mode 100644
--- /dev/null
+++ b/Engine/Source/Editor/DerivedDataWidgets/Public/DerivedDataInformation.h
@@ -0,0 +1,41 @@
+// Copyright Epic Games, Inc. All Rights Reserved.
+
+#pragma once
+
+#include "HAL/Platform.h"
+#include "Internationalization/Text.h"
+
+enum class DERIVEDDATAWIDGETS_API ERemoteCacheState : uint8
+{
+	Idle,
+	Busy,
+	Unavailable,
+	Warning,
+};
+
+class DERIVEDDATAWIDGETS_API FDerivedDataInformation
+{
+public:
+
+	static double				GetCacheActivityTimeSeconds(bool bGet, bool bLocal);
+	static double				GetCacheActivitySizeBytes(bool bGet, bool bLocal);
+	static bool					GetHasRemoteCache();
+	static bool					GetHasZenCache();
+	static bool					GetHasUnrealCloudCache();
+	static ERemoteCacheState	GetRemoteCacheState() { return RemoteCacheState; }
+	static FText				GetRemoteCacheStateAsText();
+	static FText				GetRemoteCacheWarningMessage() { return RemoteCacheWarningMessage; }
+	static void					UpdateRemoteCacheState();
+	static bool					IsUploading() { return bIsUploading; }
+	static bool					IsDownloading() { return bIsDownloading; }
+
+private:
+
+	static ERemoteCacheState	RemoteCacheState;
+	static FText				RemoteCacheWarningMessage;
+	static double				LastGetTime;
+	static double				LastPutTime;
+	static bool					bIsUploading;
+	static bool					bIsDownloading;
+
+};
diff --git a/Engine/Source/Editor/DerivedDataWidgets/Public/IDerivedDataCacheNotifications.h b/Engine/Source/Editor/DerivedDataWidgets/Public/IDerivedDataCacheNotifications.h
new file mode 100644
--- /dev/null
+++ b/Engine/Source/Editor/DerivedDataWidgets/Public/IDerivedDataCacheNotifications.h
@@ -0,0 +1,17 @@
+// Copyright Epic Games, Inc. All Rights Reserved.
+
+#pragma once
+
+#include "CoreMinimal.h"
+
+class DERIVEDDATAWIDGETS_API IDerivedDataCacheNotifications
+{
+public:
+	IDerivedDataCacheNotifications()
+	{
+	}
+
+	virtual ~IDerivedDataCacheNotifications()
+	{
+	}
+};
\ No newline at end of file
diff --git a/Engine/Source/Editor/DerivedDataWidgets/Public/SDerivedDataDialogs.h b/Engine/Source/Editor/DerivedDataWidgets/Public/SDerivedDataDialogs.h
new file mode 100644
--- /dev/null
+++ b/Engine/Source/Editor/DerivedDataWidgets/Public/SDerivedDataDialogs.h
@@ -0,0 +1,59 @@
+// Copyright Epic Games, Inc. All Rights Reserved.
+
+#pragma once
+
+#include "Templates/SharedPointer.h"
+#include "Types/SlateEnums.h"
+#include "Widgets/DeclarativeSyntaxSupport.h"
+#include "Widgets/SBoxPanel.h"
+#include "Widgets/SCompoundWidget.h"
+
+class SWidget;
+
+class DERIVEDDATAWIDGETS_API SDerivedDataRemoteStoreDialog : public SCompoundWidget
+{
+	SLATE_BEGIN_ARGS(SDerivedDataRemoteStoreDialog) {}
+	SLATE_END_ARGS()
+
+	void Construct(const FArguments& InArgs);
+
+private:
+
+	TSharedRef<SWidget> GetGridPanel();
+
+	EActiveTimerReturnType UpdateGridPanels(double InCurrentTime, float InDeltaTime);
+
+	SVerticalBox::FSlot* GridSlot = nullptr;
+};
+
+class DERIVEDDATAWIDGETS_API SDerivedDataCacheStatisticsDialog : public SCompoundWidget
+{
+	SLATE_BEGIN_ARGS(SDerivedDataCacheStatisticsDialog) {}
+	SLATE_END_ARGS()
+
+	void Construct(const FArguments& InArgs);
+
+private:
+
+	TSharedRef<SWidget> GetGridPanel();
+	
+	EActiveTimerReturnType UpdateGridPanels(double InCurrentTime, float InDeltaTime);
+
+	SVerticalBox::FSlot* GridSlot = nullptr;
+};
+
+class DERIVEDDATAWIDGETS_API SDerivedDataResourceUsageDialog : public SCompoundWidget
+{
+	SLATE_BEGIN_ARGS(SDerivedDataResourceUsageDialog) {}
+	SLATE_END_ARGS()
+
+		void Construct(const FArguments& InArgs);
+
+private:
+
+	TSharedRef<SWidget> GetGridPanel();
+	
+	EActiveTimerReturnType UpdateGridPanels(double InCurrentTime, float InDeltaTime);
+
+	SVerticalBox::FSlot* GridSlot = nullptr;
+};
diff --git a/Engine/Source/Editor/EditorStyle/Private/StarshipStyle.cpp b/Engine/Source/Editor/EditorStyle/Private/StarshipStyle.cpp
--- a/Engine/Source/Editor/EditorStyle/Private/StarshipStyle.cpp
+++ b/Engine/Source/Editor/EditorStyle/Private/StarshipStyle.cpp
@@ -1093,12 +1093,19 @@ void FStarshipEditorStyle::FStyle::SetupGeneralStyles()
 		Set("DataLayer.ColorIcon", new FSlateBoxBrush(NAME_None, 8.0f / 32.0f, FStyleColors::White));
 	}
 
+	// Zen
+	{
+		Set("Zen.Icons.LaunchDashboard", new CORE_IMAGE_BRUSH_SVG("Starship/Common/SessionFrontend", Icon16x16));
+		Set("Zen.Icons.Server", new CORE_IMAGE_BRUSH_SVG("Starship/Common/SessionFrontend", Icon16x16));
+		Set("Zen.Icons.FolderExplore", new IMAGE_BRUSH_SVG("Starship/ContentBrowser/show-in-explorer", Icon16x16));
+		Set("Zen.Icons.Clipboard", new IMAGE_BRUSH_SVG("Starship/Common/Clipboard", Icon16x16));
+	}
+
 	// Derived Data
 	{
 		Set("DerivedData.Cache.Settings", new IMAGE_BRUSH_SVG("Starship/DerivedData/DD_Cache_Settings", Icon16x16));
 		Set("DerivedData.Cache.Statistics", new IMAGE_BRUSH_SVG("Starship/DerivedData/DD_Cache_Statistics", Icon16x16));
 		Set("DerivedData.ResourceUsage", new IMAGE_BRUSH_SVG("Starship/DerivedData/DD_Resource_Usage", Icon16x16));
-		Set("DerivedData.LaunchZenDashboard", new CORE_IMAGE_BRUSH_SVG("Starship/Common/SessionFrontend", Icon16x16));
 		Set("DerivedData.RemoteCache.Uploading", new IMAGE_BRUSH_SVG("Starship/DerivedData/DD_Upload", Icon16x16, EStyleColor::AccentGreen));
 		Set("DerivedData.RemoteCache.Downloading", new IMAGE_BRUSH_SVG("Starship/DerivedData/DD_Download", Icon16x16, EStyleColor::AccentBlue));
 
diff --git a/Engine/Source/Editor/LevelEditor/LevelEditor.Build.cs b/Engine/Source/Editor/LevelEditor/LevelEditor.Build.cs
--- a/Engine/Source/Editor/LevelEditor/LevelEditor.Build.cs
+++ b/Engine/Source/Editor/LevelEditor/LevelEditor.Build.cs
@@ -80,7 +80,9 @@ public class LevelEditor : ModuleRules
 				"TranslationEditor",
 				"SubobjectEditor",
 				"SubobjectDataInterface",
+				"DerivedDataWidgets",
 				"DerivedDataEditor",
+				"ZenEditor",
 				"EditorWidgets",
 				"ToolWidgets",
 				"UnsavedAssetsTracker",
diff --git a/Engine/Source/Editor/LevelEditor/Private/SLevelEditor.cpp b/Engine/Source/Editor/LevelEditor/Private/SLevelEditor.cpp
--- a/Engine/Source/Editor/LevelEditor/Private/SLevelEditor.cpp
+++ b/Engine/Source/Editor/LevelEditor/Private/SLevelEditor.cpp
@@ -71,7 +71,8 @@
 #include "Elements/Framework/TypedElementRegistry.h"
 #include "Elements/SMInstance/SMInstanceElementId.h"
 #include "Elements/SMInstance/SMInstanceElementLevelEditorSelectionCustomization.h"
-#include "DerivedDataEditorModule.h"
+#include "DerivedDataEditor.h"
+#include "ZenEditor.h"
 #include "EditorModeManager.h"
 #include "EditorViewportLayout.h"
 #include "LevelViewportTabContent.h"
@@ -2132,14 +2133,28 @@ void SLevelEditor::RegisterStatusBarTools()
 			}
 		}));
 
+	FZenEditor& ZenEditor = FModuleManager::LoadModuleChecked<FZenEditor>("ZenEditor");
 
-	FDerivedDataEditorModule& DerivedDataEditorModule = FModuleManager::LoadModuleChecked<FDerivedDataEditorModule>("DerivedDataEditor");
+	if (ZenEditor.IsZenEnabled())
+	{
+		// Only create the Zen Status Bar tool if it is enabled. This should be the default
+		FToolMenuSection& ZenSection = Menu->AddSection("Zen", FText::GetEmpty(), FToolMenuInsert("Compile", EToolMenuInsertType::Before));
+
+		ZenSection.AddEntry(
+			FToolMenuEntry::InitWidget("ZentatusBar", ZenEditor.CreateStatusBarWidget(), FText::GetEmpty(), true, false)
+		);
+	}
+	else
+	{
+		// Optionally create the Derived Data Status Bar tool
+		FDerivedDataEditor& DerivedDataEditor = FModuleManager::LoadModuleChecked<FDerivedDataEditor>("DerivedDataEditor");
 
-	FToolMenuSection& DDCSection = Menu->AddSection("DDC", FText::GetEmpty(), FToolMenuInsert("Compile", EToolMenuInsertType::Before));
+		FToolMenuSection& DerivedDataSection = Menu->AddSection("DDC", FText::GetEmpty(), FToolMenuInsert("Compile", EToolMenuInsertType::Before));
 
-	DDCSection.AddEntry(
-		FToolMenuEntry::InitWidget("DerivedDatatatusBar", DerivedDataEditorModule.CreateStatusBarWidget(), FText::GetEmpty(), true, false)
-	);
+		DerivedDataSection.AddEntry(
+			FToolMenuEntry::InitWidget("DerivedDatatatusBar", DerivedDataEditor.CreateStatusBarWidget(), FText::GetEmpty(), true, false)
+		);
+	}
 }
 
 void SLevelEditor::AddStandaloneLevelViewport( const TSharedRef<SLevelViewport>& LevelViewport )
diff --git a/Engine/Source/Editor/UnrealEd/Private/EditorEngine.cpp b/Engine/Source/Editor/UnrealEd/Private/EditorEngine.cpp
--- a/Engine/Source/Editor/UnrealEd/Private/EditorEngine.cpp
+++ b/Engine/Source/Editor/UnrealEd/Private/EditorEngine.cpp
@@ -1502,6 +1502,7 @@ void UEditorEngine::LoadDefaultEditorModules()
 			TEXT("WorldPartitionEditor"),
 			TEXT("EditorConfig"),
 			TEXT("DerivedDataEditor"),
+			TEXT("ZenEditor"),
 			TEXT("CSVtoSVG"),
 			TEXT("GeometryFramework"),
 			TEXT("VirtualizationEditor"),
diff --git a/Engine/Source/Editor/ZenEditor/Private/SZenStatusBar.cpp b/Engine/Source/Editor/ZenEditor/Private/SZenStatusBar.cpp
new file mode 100644
--- /dev/null
+++ b/Engine/Source/Editor/ZenEditor/Private/SZenStatusBar.cpp
@@ -0,0 +1,447 @@
+// Copyright Epic Games, Inc. All Rights Reserved.
+
+#include "SZenStatusBar.h"
+
+#include "Async/Future.h"
+#include "Containers/Array.h"
+#include "Delegates/Delegate.h"
+#include "DerivedDataCacheInterface.h"
+#include "IDerivedDataCacheNotifications.h"
+#include "ZenEditor.h"
+#include "DerivedDataInformation.h"
+#include "Framework/Commands/InputChord.h"
+#include "Framework/Commands/UIAction.h"
+#include "Framework/Commands/UICommandInfo.h"
+#include "Framework/Commands/UICommandList.h"
+#include "Framework/MultiBox/MultiBoxDefs.h"
+#include "Framework/Notifications/NotificationManager.h"
+#include "Framework/SlateDelegates.h"
+#include "HAL/PlatformCrt.h"
+#include "ISettingsModule.h"
+#include "Internationalization/Internationalization.h"
+#include "Layout/Children.h"
+#include "Layout/Margin.h"
+#include "Math/Color.h"
+#include "Math/UnitConversion.h"
+#include "Math/UnrealMathSSE.h"
+#include "Misc/Attribute.h"
+#include "Misc/CoreMisc.h"
+#include "Modules/ModuleManager.h"
+#include "SlotBase.h"
+#include "Styling/AppStyle.h"
+#include "Styling/ISlateStyle.h"
+#include "Styling/SlateColor.h"
+#include "Styling/SlateTypes.h"
+#include "Textures/SlateIcon.h"
+#include "ToolMenu.h"
+#include "ToolMenuContext.h"
+#include "ToolMenuSection.h"
+#include "ToolMenus.h"
+#include "Types/WidgetActiveTimerDelegate.h"
+#include "UObject/NameTypes.h"
+#include "UObject/UnrealNames.h"
+#include "Widgets/Images/SImage.h"
+#include "Widgets/Input/SComboButton.h"
+#include "Widgets/Notifications/SNotificationList.h"
+#include "Widgets/SBoxPanel.h"
+#include "Widgets/SOverlay.h"
+#include "Widgets/Text/STextBlock.h"
+#include "ZenDashboardLauncher.h"
+
+class SWidget;
+struct FSlateBrush;
+
+#define LOCTEXT_NAMESPACE "ZenStatusBar"
+
+
+TSharedRef<FUICommandList> FZenStausBarCommands::ActionList(new FUICommandList());
+
+FZenStausBarCommands::FZenStausBarCommands()
+	: TCommands<FZenStausBarCommands>
+	(
+		"ZenSettings",
+		NSLOCTEXT("Contexts", "Zen", "Zen"),
+		"LevelEditor",
+		FAppStyle::GetAppStyleSetName()
+	)
+{}
+
+void FZenStausBarCommands::RegisterCommands()
+{
+	UI_COMMAND(ChangeSettings, "Change Cache Settings", "Opens a dialog to change Cache settings.", EUserInterfaceActionType::Button, FInputChord());
+	UI_COMMAND(ViewCacheStatistics, "Cache Statistics", "Opens the Cache Statistics panel.", EUserInterfaceActionType::Button, FInputChord());
+	UI_COMMAND(ViewResourceUsage, "Cache Resource Usage", "Opens the Cache Resource Usage panel.", EUserInterfaceActionType::Button, FInputChord());
+	UI_COMMAND(ViewServerStatus, "Server Status", "Opens the Zen Server Status panel.", EUserInterfaceActionType::Button, FInputChord());
+	UI_COMMAND(LaunchZenDashboard, "Launch Dashboard", "Launches the Zen Dashboard utility.", EUserInterfaceActionType::Button, FInputChord());
+
+	ActionList->MapAction(
+		ChangeSettings,
+		FExecuteAction::CreateStatic(&FZenStausBarCommands::ChangeSettings_Clicked)
+	);
+
+	ActionList->MapAction(
+		ViewCacheStatistics,
+		FExecuteAction::CreateStatic(&FZenStausBarCommands::ViewCacheStatistics_Clicked)
+	);
+
+	ActionList->MapAction(
+		ViewResourceUsage,
+		FExecuteAction::CreateStatic(&FZenStausBarCommands::ViewResourceUsage_Clicked)
+	);
+
+	ActionList->MapAction(
+		ViewServerStatus,
+		FExecuteAction::CreateStatic(&FZenStausBarCommands::ViewServerStatus_Clicked)
+	);
+
+	ActionList->MapAction(
+		LaunchZenDashboard,
+		FExecuteAction::CreateStatic(&FZenStausBarCommands::LaunchZenDashboard_Clicked)
+	);
+}
+
+void FZenStausBarCommands::ChangeSettings_Clicked()
+{
+	FModuleManager::LoadModuleChecked<ISettingsModule>("Settings").ShowViewer("Editor", "General", "Global");
+}
+
+void FZenStausBarCommands::ViewCacheStatistics_Clicked()
+{
+	FZenEditor& ZenEditor = FModuleManager::LoadModuleChecked<FZenEditor>("ZenEditor");
+	ZenEditor.ShowCacheStatisticsTab();
+}
+
+void FZenStausBarCommands::ViewResourceUsage_Clicked()
+{
+	FZenEditor& ZenEditor = FModuleManager::LoadModuleChecked<FZenEditor>("ZenEditor");
+	ZenEditor.ShowResourceUsageTab();
+}
+
+void FZenStausBarCommands::ViewServerStatus_Clicked()
+{
+	FZenEditor& ZenEditor = FModuleManager::LoadModuleChecked<FZenEditor>("ZenEditor");
+	ZenEditor.ShowZenServerStatusTab();
+}
+
+void FZenStausBarCommands::LaunchZenDashboard_Clicked()
+{
+	using namespace UE::Zen;
+	FZenDashboardLauncher::Get()->StartZenDashboard(FZenDashboardLauncher::Get()->GetZenDashboardApplicationPath());
+}
+
+TSharedRef<SWidget> SZenStatusBarWidget::CreateStatusBarMenu()
+{
+	UToolMenu* Menu = UToolMenus::Get()->RegisterMenu("StatusBar.ToolBar.Zen", NAME_None, EMultiBoxType::Menu, false);
+
+	{
+		FToolMenuSection& Section = Menu->AddSection("ZenMenuToolsSection", LOCTEXT("ZenMenuToolsSection", "Tools"));
+
+		Section.AddMenuEntry(
+			FZenStausBarCommands::Get().LaunchZenDashboard,
+			TAttribute<FText>(),
+			TAttribute<FText>(),
+			FSlateIcon(FAppStyle::GetAppStyleSetName(), "Zen.Icons.LaunchDashboard")
+		);
+	}
+
+	{
+		FToolMenuSection& Section = Menu->AddSection("ZenMenuServerSection", LOCTEXT("ZenMenuServerSection", "Server"));
+
+		Section.AddMenuEntry(
+			FZenStausBarCommands::Get().ViewServerStatus,
+			TAttribute<FText>(),
+			TAttribute<FText>(),
+			FSlateIcon(FAppStyle::GetAppStyleSetName(), "Zen.Icons.Server")
+		);
+	}
+
+	{
+		FToolMenuSection& Section = Menu->AddSection("ZenMenuDerivedDataSection", LOCTEXT("ZenMenuDerivedDataSection", "Cache"));
+
+		Section.AddMenuEntry(
+			FZenStausBarCommands::Get().ChangeSettings,
+			TAttribute<FText>(),
+			TAttribute<FText>(),
+			FSlateIcon(FAppStyle::GetAppStyleSetName(), "DerivedData.Cache.Settings")
+		);
+
+		Section.AddMenuEntry(
+			FZenStausBarCommands::Get().ViewCacheStatistics,
+			TAttribute<FText>(),
+			TAttribute<FText>(),
+			FSlateIcon(FAppStyle::GetAppStyleSetName(), "DerivedData.Cache.Statistics")
+		);
+	
+		Section.AddMenuEntry(
+			FZenStausBarCommands::Get().ViewResourceUsage,
+			TAttribute<FText>(),
+			TAttribute<FText>(),
+			FSlateIcon(FAppStyle::GetAppStyleSetName(), "DerivedData.ResourceUsage")
+		);
+	}
+
+	return UToolMenus::Get()->GenerateWidget("StatusBar.ToolBar.Zen", FToolMenuContext(FZenStausBarCommands::ActionList));
+}
+
+void SZenStatusBarWidget::Construct(const FArguments& InArgs)
+{	
+	this->ChildSlot
+	[
+		SNew(SComboButton)
+		.ContentPadding(FMargin(6.0f, 0.0f))
+		.MenuPlacement(MenuPlacement_AboveAnchor)
+		.ComboButtonStyle(&FAppStyle::Get().GetWidgetStyle<FComboButtonStyle>("SimpleComboButton"))
+		.ButtonContent()
+		[
+
+			SNew(SHorizontalBox)
+	
+			+ SHorizontalBox::Slot()
+			.AutoWidth()
+			.VAlign(VAlign_Center)
+			.Padding(0, 0, 3, 0)
+			[
+				SNew(SOverlay)
+				+ SOverlay::Slot()
+				.HAlign(HAlign_Center)
+				.VAlign(VAlign_Top)
+				[
+					SNew(SImage)
+					.ColorAndOpacity(FSlateColor::UseForeground())
+					.Image_Lambda([this] { return GetRemoteCacheStateBackgroundIcon();  })
+					.ToolTipText_Lambda([this] { return GetRemoteCacheToolTipText(); })
+				]
+
+				+ SOverlay::Slot()
+				.HAlign(HAlign_Center)
+				.VAlign(VAlign_Top)
+				[
+					SNew(SImage)
+					.ColorAndOpacity(FSlateColor::UseForeground())
+					.Image_Lambda([this] { return GetRemoteCacheStateBadgeIcon();  })
+					.ToolTipText_Lambda([this] { return GetRemoteCacheToolTipText(); })
+				]
+
+				+ SOverlay::Slot()
+				.HAlign(HAlign_Center)
+				.VAlign(VAlign_Top)
+				[
+					SNew(SImage)
+					.Image(FAppStyle::Get().GetBrush("DerivedData.RemoteCache.Uploading"))
+					.ColorAndOpacity_Lambda([this] { return ( FDerivedDataInformation::IsUploading() && FDerivedDataInformation::GetRemoteCacheState() == ERemoteCacheState::Busy )? FLinearColor::White.CopyWithNewOpacity(FMath::MakePulsatingValue(ElapsedUploadTime, 2)) : FLinearColor(0,0,0,0); })
+					.ToolTipText_Lambda([this] { return GetRemoteCacheToolTipText(); })
+				]
+
+				+ SOverlay::Slot()
+				.HAlign(HAlign_Center)
+				.VAlign(VAlign_Top)
+				[
+					SNew(SImage)
+					.Image(FAppStyle::Get().GetBrush("DerivedData.RemoteCache.Downloading"))
+					.ColorAndOpacity_Lambda([this] { return ( FDerivedDataInformation::IsDownloading() && FDerivedDataInformation::GetRemoteCacheState()==ERemoteCacheState::Busy ) ? FLinearColor::White.CopyWithNewOpacity(FMath::MakePulsatingValue(ElapsedDownloadTime, 2)) : FLinearColor(0, 0, 0, 0); })
+					.ToolTipText_Lambda([this] { return GetRemoteCacheToolTipText(); })
+				]				
+			]
+
+			+ SHorizontalBox::Slot()
+			.AutoWidth()
+			.VAlign(VAlign_Center)
+			.Padding(0, 0, 10, 0)
+			[
+				SNew(STextBlock)
+				.Text_Lambda([this] { return GetTitleText(); })
+				.ToolTipText_Lambda([this] { return GetTitleToolTipText(); })
+			]
+		]
+		.OnGetMenuContent(FOnGetContent::CreateRaw(this, &SZenStatusBarWidget::CreateStatusBarMenu))
+	];
+
+	RegisterActiveTimer(0.2f, FWidgetActiveTimerDelegate::CreateSP(this, &SZenStatusBarWidget::UpdateBusyIndicator));
+	RegisterActiveTimer(5.0f, FWidgetActiveTimerDelegate::CreateSP(this, &SZenStatusBarWidget::UpdateWarnings));
+}
+
+EActiveTimerReturnType SZenStatusBarWidget::UpdateBusyIndicator(double InCurrentTime, float InDeltaTime)
+{	
+	FDerivedDataInformation::UpdateRemoteCacheState();
+
+	bBusy = GetDerivedDataCache()->AnyAsyncRequestsRemaining();
+
+	if (FDerivedDataInformation::IsUploading())
+	{
+		ElapsedUploadTime += fmod(InDeltaTime,3600.0);
+	}
+	else
+	{
+		ElapsedUploadTime = 0.0;
+	}
+
+	if (FDerivedDataInformation::IsDownloading())
+	{
+		ElapsedDownloadTime += fmod(InDeltaTime, 3600.0);
+	}
+	else
+	{
+		ElapsedDownloadTime = 0.0;
+	}
+
+	if (bBusy)
+	{
+		ElapsedBusyTime += fmod(InDeltaTime, 3600.0);
+	}
+	else
+	{
+		ElapsedBusyTime = 0;
+	}
+
+	return EActiveTimerReturnType::Continue;
+}
+
+EActiveTimerReturnType SZenStatusBarWidget::UpdateWarnings(double InCurrentTime, float InDeltaTime)
+{
+	if ( FDerivedDataInformation::GetRemoteCacheState()== ERemoteCacheState::Warning )
+	{
+		if ( NotificationItem.IsValid() == false || NotificationItem->GetCompletionState()== SNotificationItem::CS_None)
+		{
+			// No existing notification or the existing one has finished
+			TPromise<TWeakPtr<SNotificationItem>> NotificationPromise;
+
+			FNotificationInfo Info(FDerivedDataInformation::GetRemoteCacheWarningMessage());
+			Info.bUseSuccessFailIcons = true;
+			Info.bFireAndForget = false;
+			Info.bUseThrobber = false;
+			Info.FadeOutDuration = 0.0f;
+			Info.ExpireDuration = 0.0f;
+
+			Info.ButtonDetails.Add(FNotificationButtonInfo(LOCTEXT("OpenSettings", "Open Settings"), FText(), FSimpleDelegate::CreateLambda([NotificationFuture = NotificationPromise.GetFuture().Share()]() {
+				FModuleManager::LoadModuleChecked<ISettingsModule>("Settings").ShowViewer("Editor", "General", "Global");
+
+				TWeakPtr<SNotificationItem> NotificationPtr = NotificationFuture.Get();
+				if (TSharedPtr<SNotificationItem> Notification = NotificationPtr.Pin())
+				{
+					Notification->SetCompletionState(SNotificationItem::CS_None);
+					Notification->ExpireAndFadeout();
+				}
+			}),
+				SNotificationItem::ECompletionState::CS_Fail));
+
+			NotificationItem = FSlateNotificationManager::Get().AddNotification(Info);
+
+			if (NotificationItem.IsValid())
+			{
+				NotificationPromise.SetValue(NotificationItem);
+				NotificationItem->SetCompletionState(SNotificationItem::CS_Fail);
+			}
+		}
+	}
+	else
+	{
+		// No longer any warnings
+		if (NotificationItem.IsValid())
+		{
+			NotificationItem->SetCompletionState(SNotificationItem::CS_None);
+			NotificationItem->ExpireAndFadeout();
+		}
+	}
+
+	return EActiveTimerReturnType::Continue;
+}
+
+FText SZenStatusBarWidget::GetTitleToolTipText() const
+{
+	FTextBuilder DescBuilder;
+
+	DescBuilder.AppendLineFormat(LOCTEXT("GraphNameText", "Graph : {0}"), FText::FromString(GetDerivedDataCache()->GetGraphName()));
+	
+	return DescBuilder.ToText();
+}
+
+FText SZenStatusBarWidget::GetTitleText() const
+{
+	return LOCTEXT("ZenStatusBarName", "Zen Server");
+}
+
+FText SZenStatusBarWidget::GetRemoteCacheToolTipText() const
+{
+	FTextBuilder DescBuilder;
+
+	if (FDerivedDataInformation::GetRemoteCacheState() == ERemoteCacheState::Warning)
+	{
+		DescBuilder.AppendLineFormat(LOCTEXT("RemoteCacheErrorText", "WARNING\t: {0}\n"), FDerivedDataInformation::GetRemoteCacheWarningMessage());
+	}
+
+	DescBuilder.AppendLine(LOCTEXT("RemoteCacheToolTipText", "Remote Cache\n"));
+	DescBuilder.AppendLineFormat(LOCTEXT("RemoteCacheConnectedText", "Connected\t: {0}"), FText::FromString((FDerivedDataInformation::GetHasRemoteCache() ? TEXT("Yes") : TEXT("No"))));
+	DescBuilder.AppendLineFormat(LOCTEXT("RemoteCacheStatusText", "Status\t: {0}"), FDerivedDataInformation::GetRemoteCacheStateAsText() );
+	
+	const double DownloadedBytesMB = FUnitConversion::Convert(FDerivedDataInformation::GetCacheActivitySizeBytes(true, false), EUnit::Bytes, EUnit::Megabytes);
+	const double UploadedBytesMB = FUnitConversion::Convert(FDerivedDataInformation::GetCacheActivitySizeBytes(false, false), EUnit::Bytes, EUnit::Megabytes);
+
+	DescBuilder.AppendLineFormat(LOCTEXT("RemoteCacheDownloaded", "Downloaded\t: {0} MB"), DownloadedBytesMB);
+	DescBuilder.AppendLineFormat(LOCTEXT("RemoteCacheUploaded",	"Uploaded\t: {0} MB"), UploadedBytesMB);
+	
+	return DescBuilder.ToText();
+}
+
+const FSlateBrush* SZenStatusBarWidget::GetRemoteCacheStateBackgroundIcon() const
+{
+	switch ( FDerivedDataInformation::GetRemoteCacheState())
+	{
+		case ERemoteCacheState::Idle :
+		{
+			return FAppStyle::Get().GetBrush("DerivedData.RemoteCache.IdleBG");
+			break;
+		}
+
+		case ERemoteCacheState::Busy :
+		{
+			return FAppStyle::Get().GetBrush("DerivedData.RemoteCache.BusyBG");
+			break;
+		}	
+
+		case ERemoteCacheState::Unavailable:
+		{
+			return FAppStyle::Get().GetBrush("DerivedData.RemoteCache.UnavailableBG");
+			break;
+		}
+
+		default:
+		case ERemoteCacheState::Warning:
+		{
+			return FAppStyle::Get().GetBrush("DerivedData.RemoteCache.WarningBG");
+			break;
+		}
+	}
+}
+
+const FSlateBrush* SZenStatusBarWidget::GetRemoteCacheStateBadgeIcon() const
+{
+	switch ( FDerivedDataInformation::GetRemoteCacheState())
+	{
+		case ERemoteCacheState::Idle :
+		{
+			return FAppStyle::Get().GetBrush("DerivedData.RemoteCache.Idle");
+			break;
+		}
+
+		case ERemoteCacheState::Busy :
+		{
+			return FAppStyle::Get().GetBrush("DerivedData.RemoteCache.Busy");
+			break;
+		}	
+
+		case ERemoteCacheState::Unavailable:
+		{
+			return FAppStyle::Get().GetBrush("DerivedData.RemoteCache.Unavailable");
+			break;
+		}
+
+		default:
+		case ERemoteCacheState::Warning:
+		{
+			return FAppStyle::Get().GetBrush("DerivedData.RemoteCache.Warning");
+			break;
+		}
+	}
+}
+
+#undef LOCTEXT_NAMESPACE
\ No newline at end of file
diff --git a/Engine/Source/Editor/ZenEditor/Private/SZenStatusBar.h b/Engine/Source/Editor/ZenEditor/Private/SZenStatusBar.h
new file mode 100644
--- /dev/null
+++ b/Engine/Source/Editor/ZenEditor/Private/SZenStatusBar.h
@@ -0,0 +1,71 @@
+// Copyright Epic Games, Inc. All Rights Reserved.
+
+#pragma once
+
+#include "Framework/Commands/Commands.h"
+#include "Internationalization/Text.h"
+#include "Templates/SharedPointer.h"
+#include "Types/SlateEnums.h"
+#include "Widgets/DeclarativeSyntaxSupport.h"
+#include "Widgets/SCompoundWidget.h"
+
+class FUICommandInfo;
+class FUICommandList;
+class SNotificationItem;
+class SWidget;
+struct FSlateBrush;
+
+class FZenStausBarCommands : public TCommands<FZenStausBarCommands>
+{
+public:
+
+	FZenStausBarCommands();
+
+	virtual void RegisterCommands() override;
+
+private:
+
+	static void ChangeSettings_Clicked();
+	static void ViewCacheStatistics_Clicked();
+	static void ViewResourceUsage_Clicked();
+	static void ViewServerStatus_Clicked();
+	static void LaunchZenDashboard_Clicked();
+
+public:
+
+	TSharedPtr< FUICommandInfo > ChangeSettings;
+	TSharedPtr< FUICommandInfo > ViewResourceUsage;
+	TSharedPtr< FUICommandInfo > ViewCacheStatistics;
+	TSharedPtr< FUICommandInfo > ViewServerStatus;
+	TSharedPtr< FUICommandInfo > LaunchZenDashboard;
+
+	static TSharedRef<FUICommandList> ActionList;
+};
+
+
+class SZenStatusBarWidget : public SCompoundWidget
+{
+	SLATE_BEGIN_ARGS(SZenStatusBarWidget) {}
+	SLATE_END_ARGS()
+
+	void Construct(const FArguments& InArgs);
+
+private:
+
+	FText						GetTitleToolTipText() const;
+	FText						GetRemoteCacheToolTipText() const;
+	FText						GetTitleText() const;
+	const FSlateBrush*			GetRemoteCacheStateBackgroundIcon() const;
+	const FSlateBrush*			GetRemoteCacheStateBadgeIcon() const;
+	TSharedRef<SWidget>			CreateStatusBarMenu();
+	EActiveTimerReturnType		UpdateBusyIndicator(double InCurrentTime, float InDeltaTime);
+	EActiveTimerReturnType		UpdateWarnings(double InCurrentTime, float InDeltaTime);
+
+	double ElapsedDownloadTime = 0;
+	double ElapsedUploadTime = 0;
+	double ElapsedBusyTime = 0;
+
+	bool bBusy = false;
+
+	TSharedPtr<SNotificationItem> NotificationItem;
+};
diff --git a/Engine/Source/Editor/ZenEditor/Private/ZenDashboardLauncher.cpp b/Engine/Source/Editor/ZenEditor/Private/ZenDashboardLauncher.cpp
new file mode 100644
--- /dev/null
+++ b/Engine/Source/Editor/ZenEditor/Private/ZenDashboardLauncher.cpp
@@ -0,0 +1,129 @@
+// Copyright Epic Games, Inc. All Rights Reserved.
+
+#include "ZenDashboardLauncher.h"
+
+#if UE_WITH_ZEN
+
+#include "IUATHelperModule.h"
+#include "Logging/LogMacros.h"
+#include "Logging/MessageLog.h"
+#include "Misc/Paths.h"
+#include "Styling/AppStyle.h"
+
+#define LOCTEXT_NAMESPACE "FZenDashboardLauncher"
+
+namespace UE::Zen
+{
+
+DEFINE_LOG_CATEGORY_STATIC(LogZenDashboardLauncher, Log, All);
+
+TSharedPtr<FZenDashboardLauncher> FZenDashboardLauncher::Instance = nullptr;
+
+FZenDashboardLauncher::FZenDashboardLauncher()
+{
+
+}
+
+FZenDashboardLauncher::~FZenDashboardLauncher()
+{
+
+}
+	
+FString FZenDashboardLauncher::GetZenDashboardApplicationPath()
+{
+	FString Path = FPlatformProcess::GenerateApplicationPath(TEXT("ZenDashboard"), EBuildConfiguration::Development);
+	return FPaths::ConvertRelativePathToFull(Path);
+}
+
+void FZenDashboardLauncher::StartZenDashboard(const FString& Path, const FString& Parameters)
+{
+	auto Callback = [](const EStartZenDashboardResult Result) {};
+	StartZenDashboard(Path, Parameters, Callback);
+}
+
+void FZenDashboardLauncher::StartZenDashboard(const FString& Path, const FString& Parameters, StartZenDashboardCallback Callback)
+{
+    if (!FPaths::FileExists(Path))
+    {
+    	BuildZenDashboard(Path, Parameters, Callback);
+    	return;
+    }
+	
+	constexpr bool bLaunchDetached = true;
+	constexpr bool bLaunchHidden = false;
+	constexpr bool bLaunchReallyHidden = false;
+
+	uint32 ProcessID = 0;
+	const int32 PriorityModifier = 0;
+
+	const TCHAR* OptionalWorkingDirectory = nullptr;
+
+	void* PipeWriteChild = nullptr;
+	void* PipeReadChild = nullptr;
+	ZenDashboardHandle = FPlatformProcess::CreateProc(*Path, *Parameters, bLaunchDetached, bLaunchHidden, bLaunchReallyHidden, &ProcessID, PriorityModifier, OptionalWorkingDirectory, PipeWriteChild, PipeReadChild);
+
+	if (ZenDashboardHandle.IsValid())
+	{
+		UE_LOG(LogZenDashboardLauncher, Log, TEXT("Launched Zen Dashboard executable: %s %s"), *Path, *Parameters);
+		Callback(EStartZenDashboardResult::Completed);
+	}
+	else
+	{
+		Callback(EStartZenDashboardResult::LaunchFailed);
+	}
+}
+
+void FZenDashboardLauncher::CloseZenDashboard()
+{
+	if (ZenDashboardHandle.IsValid())
+	{
+		FPlatformProcess::TerminateProc(ZenDashboardHandle);
+	}
+	else
+	{
+		UE_LOG(LogZenDashboardLauncher, Log, TEXT("Could not find the Zen Dashboard process handler"));
+	}
+}
+
+void FZenDashboardLauncher::BuildZenDashboard(const FString& Path, const FString& LaunchParameters, StartZenDashboardCallback Callback)
+{
+	UE_LOG(LogZenDashboardLauncher, Log, TEXT("Could not find the Zen Dashboard executable: %s. Attempting to build Zen Dashboard."), *Path);
+
+	FString Arguments;
+#if PLATFORM_WINDOWS
+	FText PlatformName = LOCTEXT("PlatformName_Windows", "Windows");
+	Arguments = TEXT("BuildTarget -Target=ZenDashboard -Platform=Win64");
+#elif PLATFORM_MAC
+	FText PlatformName = LOCTEXT("PlatformName_Mac", "Mac");
+	Arguments = TEXT("BuildTarget -Target=ZenDashboard -Platform=Mac");
+#elif PLATFORM_LINUX
+	FText PlatformName = LOCTEXT("PlatformName_Linux", "Linux");
+	Arguments = TEXT("BuildTarget -Target=ZenDashboard -Platform=Linux");
+#endif
+
+	IUATHelperModule::Get().CreateUatTask(Arguments, PlatformName, LOCTEXT("BuildingZenDashboard", "Building Zen Dashboard"),
+		LOCTEXT("BuildZenDashboardTask", "Build Zen Dashboard Task"), FAppStyle::GetBrush(TEXT("MainFrame.CookContent")), nullptr,
+		[this, Path, LaunchParameters, Callback](FString Result, double Time)
+		{
+			if (Result.Equals(TEXT("Completed")))
+			{
+#if PLATFORM_MAC
+				// On Mac we genereate the path again so that it includes the newly built executable.
+				FString NewPath = GetZenDashboardApplicationPath();
+				this->StartZenDashboard(NewPath, LaunchParameters, Callback);
+#else
+				this->StartZenDashboard(Path, LaunchParameters, Callback);
+#endif
+			}
+			else
+			{
+				Callback(EStartZenDashboardResult::BuildFailed);
+			}
+		});
+}
+
+} // namespace UE::Zen
+
+#undef LOCTEXT_NAMESPACE
+
+#endif // UE_WITH_ZEN
\ No newline at end of file
diff --git a/Engine/Source/Editor/ZenEditor/Private/ZenDashboardLauncher.h b/Engine/Source/Editor/ZenEditor/Private/ZenDashboardLauncher.h
new file mode 100644
--- /dev/null
+++ b/Engine/Source/Editor/ZenEditor/Private/ZenDashboardLauncher.h
@@ -0,0 +1,85 @@
+// Copyright Epic Games, Inc. All Rights Reserved.
+
+#pragma once
+
+#include "CoreMinimal.h"
+#include "Experimental/ZenGlobals.h"
+#include "HAL/PlatformProcess.h"
+
+#if UE_WITH_ZEN
+
+#define UE_API
+
+namespace UE::Zen
+{
+
+class ZENEDITOR_API FZenDashboardLauncher : public TSharedFromThis<FZenDashboardLauncher>
+{
+public:
+	enum class EStartZenDashboardResult : uint32
+	{
+		Completed = 0,
+		BuildFailed = 1,
+		LaunchFailed = 2,
+	};
+
+	typedef TFunction<void(const EStartZenDashboardResult /*Result*/)> StartZenDashboardCallback;
+
+	UE_API FZenDashboardLauncher();
+	UE_API ~FZenDashboardLauncher();
+	
+	static UE_API const TSharedPtr<FZenDashboardLauncher>& Get()
+	{
+		if (!Instance)
+		{
+			Instance = MakeShared<FZenDashboardLauncher>();
+		}
+		return Instance;
+	}
+
+	/**
+	 * Returns the full path to ZenDashboard.exe for the current engine installation 
+	 */
+	UE_API FString GetZenDashboardApplicationPath();
+
+	/**
+	 * Launches the ZenDashboard executable from the given Path, displays an editor message if it fails.
+	 * If the executable is not found, a build process is started.
+	 * @param Path The full filename of ZenDashboard.exe to launch
+	 * @param Parameters The command line parameters to use when launching the exe
+	 */	
+	UE_API void StartZenDashboard(const FString& Path, const FString& Parameters = TEXT(""));
+
+	/**
+	* Launches the ZenDashboard executable from the given Path, displays an editor message if it fails.
+	* If the executable is not found, a build process is started.
+	* @param Path The full filename of ZenDashboard.exe to launch
+	* @param Parameters The command line parameters to use when launching the exe
+	* @param Callback A Callback that will be called when the launch process is completed.
+	*/
+	UE_API void StartZenDashboard(const FString& Path, const FString& Parameters, StartZenDashboardCallback Callback);
+
+	/**
+	* Closes ZenDashboard.exe.
+	*/
+	UE_API void CloseZenDashboard();
+
+private:
+	/*
+	 * Attempts building ZenDashboard via UAT, will launch with forwarded parameters if successful.
+	 * Assumes that the ZenDashboard Executable in path belongs to this engine.
+	 */
+	void BuildZenDashboard(const FString& Path, const FString& LaunchParameters, StartZenDashboardCallback Callback);
+
+private:
+	/** The process handler of ZenDashboard. */
+	FProcHandle ZenDashboardHandle;
+
+	static TSharedPtr<FZenDashboardLauncher> Instance;
+};
+
+} // namespace UE::Zen
+
+#undef UE_API
+
+#endif // UE_WITH_ZEN
diff --git a/Engine/Source/Editor/ZenEditor/Private/ZenEditor.cpp b/Engine/Source/Editor/ZenEditor/Private/ZenEditor.cpp
new file mode 100644
--- /dev/null
+++ b/Engine/Source/Editor/ZenEditor/Private/ZenEditor.cpp
@@ -0,0 +1,164 @@
+// Copyright Epic Games, Inc. All Rights Reserved.
+
+#include "ZenEditor.h"
+#include "DerivedDataCacheNotifications.h"
+#include "IDerivedDataCacheNotifications.h"
+#include "DerivedDataInformation.h"
+#include "Experimental/ZenServerInterface.h"
+#include "Framework/Application/SlateApplication.h"
+#include "Modules/ModuleManager.h"
+#include "SDerivedDataDialogs.h"
+#include "SZenDialogs.h"
+#include "SZenStatusBar.h"
+#include "Styling/AppStyle.h"
+#include "Widgets/DeclarativeSyntaxSupport.h"
+#include "Widgets/Docking/SDockTab.h"
+#include "WorkspaceMenuStructure.h"
+#include "WorkspaceMenuStructureModule.h"
+
+#define LOCTEXT_NAMESPACE "ZenEditor"
+
+IMPLEMENT_MODULE(FZenEditor, ZenEditor);
+
+static const FName DerivedDataResourceUsageTabName = FName(TEXT("DerivedDataResourceUsage"));
+static const FName DerivedDataCacheStatisticsTabName = FName(TEXT("DerivedDataCacheStatistics"));
+static const FName ZenServerStatusTabName = FName(TEXT("ZenServerStatus"));
+
+void FZenEditor::StartupModule()
+{
+	const FSlateIcon ResourceUsageIcon(FAppStyle::GetAppStyleSetName(), "DerivedData.ResourceUsage");
+
+	FGlobalTabmanager::Get()->RegisterNomadTabSpawner(DerivedDataResourceUsageTabName, FOnSpawnTab::CreateRaw(this, &FZenEditor::CreateResourceUsageTab))
+		.SetDisplayName(LOCTEXT("DerivedDataResourceUsageTabTitle", "Resource Usage"))
+		.SetTooltipText(LOCTEXT("DerivedDataResourceUsageTabToolTipText", "Derived Data Resource Usage"))
+		.SetGroup(WorkspaceMenu::GetMenuStructure().GetToolsCategory())
+		.SetIcon(ResourceUsageIcon);
+
+	const FSlateIcon CacheStatisticsIcon(FAppStyle::GetAppStyleSetName(), "DerivedData.Cache.Statistics");
+
+	FGlobalTabmanager::Get()->RegisterNomadTabSpawner(DerivedDataCacheStatisticsTabName, FOnSpawnTab::CreateRaw(this, &FZenEditor::CreateCacheStatisticsTab))
+		.SetDisplayName(LOCTEXT("DerivedDataCacheStatisticsTabTitle", "Cache Statistics"))
+		.SetTooltipText(LOCTEXT("DerivedDataCacheStatisticsTabToolTipText", "Derived Data Cache Statistics"))
+		.SetGroup(WorkspaceMenu::GetMenuStructure().GetToolsCategory())
+		.SetIcon(CacheStatisticsIcon);
+
+	const FSlateIcon ZenServerIcon(FAppStyle::GetAppStyleSetName(), "Zen.Icons.Server");
+
+	FGlobalTabmanager::Get()->RegisterNomadTabSpawner(ZenServerStatusTabName, FOnSpawnTab::CreateRaw(this, &FZenEditor::CreateZenServerStatusTab))
+		.SetDisplayName(LOCTEXT("ZenServerStatusTabTitle", "Zen Server Status"))
+		.SetTooltipText(LOCTEXT("ZenServerStatusTabToolTipText", "Zen Server Status"))
+		.SetGroup(WorkspaceMenu::GetMenuStructure().GetToolsCategory())
+		.SetIcon(ZenServerIcon);
+
+#if WITH_RELOAD
+	// This code attempts to relaunch the tabs when you reload this module
+	if (IsReloadActive() && FSlateApplication::IsInitialized())
+	{
+		ShowCacheStatisticsTab();
+		ShowResourceUsageTab();
+		ShowZenServerStatusTab();
+	}
+#endif // WITH_RELOAD
+
+	FZenStausBarCommands::Register();
+
+	DerivedDataCacheNotifications.Reset(new FDerivedDataCacheNotifications);
+}
+
+void FZenEditor::ShutdownModule()
+{
+	if (FSlateApplication::IsInitialized())
+	{
+		FGlobalTabmanager::Get()->UnregisterNomadTabSpawner(DerivedDataResourceUsageTabName);
+
+		if (ResourceUsageTab.IsValid())
+		{
+			ResourceUsageTab.Pin()->RequestCloseTab();
+		}
+
+		FGlobalTabmanager::Get()->UnregisterNomadTabSpawner(DerivedDataCacheStatisticsTabName);
+
+		if (CacheStatisticsTab.IsValid())
+		{
+			CacheStatisticsTab.Pin()->RequestCloseTab();
+		}
+
+		FGlobalTabmanager::Get()->UnregisterNomadTabSpawner(ZenServerStatusTabName);
+
+		if (ZenServerStatusTab.IsValid())
+		{
+			ZenServerStatusTab.Pin()->RequestCloseTab();
+		}
+	}
+
+	FZenStausBarCommands::Unregister();
+}
+
+bool FZenEditor::IsZenEnabled() const
+{
+	return UE::Zen::IsDefaultServicePresent();
+}
+
+TSharedRef<SWidget> FZenEditor::CreateStatusBarWidget()
+{
+	return SNew(SZenStatusBarWidget);
+}
+
+TSharedPtr<SWidget> FZenEditor::CreateResourceUsageDialog()
+{
+	return SNew(SDerivedDataResourceUsageDialog);
+}
+
+TSharedRef<SDockTab> FZenEditor::CreateResourceUsageTab(const FSpawnTabArgs& Args)
+{
+	return SAssignNew(ResourceUsageTab, SDockTab)
+	.TabRole(ETabRole::NomadTab)
+	[
+		CreateResourceUsageDialog().ToSharedRef()
+	];
+}
+
+void FZenEditor::ShowResourceUsageTab()
+{
+	FGlobalTabmanager::Get()->TryInvokeTab(FTabId(DerivedDataResourceUsageTabName));
+}
+
+TSharedPtr<SWidget> FZenEditor::CreateCacheStatisticsDialog()
+{
+	return SNew(SDerivedDataCacheStatisticsDialog);
+}
+
+TSharedRef<SDockTab> FZenEditor::CreateCacheStatisticsTab(const FSpawnTabArgs& Args)
+{
+	return SAssignNew(CacheStatisticsTab, SDockTab)
+		.TabRole(ETabRole::NomadTab)
+		[
+			CreateCacheStatisticsDialog().ToSharedRef()
+		];
+}
+
+void FZenEditor::ShowCacheStatisticsTab()
+{
+	FGlobalTabmanager::Get()->TryInvokeTab(FTabId(DerivedDataCacheStatisticsTabName));
+}
+
+TSharedPtr<SWidget> FZenEditor::CreateZenStoreDialog()
+{
+	return SNew(SZenStoreStausDialog);
+}
+
+void FZenEditor::ShowZenServerStatusTab()
+{
+	FGlobalTabmanager::Get()->TryInvokeTab(FTabId(ZenServerStatusTabName));
+}
+
+TSharedRef<SDockTab> FZenEditor::CreateZenServerStatusTab(const FSpawnTabArgs& Args)
+{
+	return SAssignNew(ZenServerStatusTab, SDockTab)
+		.TabRole(ETabRole::NomadTab)
+		[
+			CreateZenStoreDialog().ToSharedRef()
+		];
+}
+
+#undef LOCTEXT_NAMESPACE
diff --git a/Engine/Source/Editor/ZenEditor/Public/ZenEditor.h b/Engine/Source/Editor/ZenEditor/Public/ZenEditor.h
new file mode 100644
--- /dev/null
+++ b/Engine/Source/Editor/ZenEditor/Public/ZenEditor.h
@@ -0,0 +1,62 @@
+// Copyright Epic Games, Inc. All Rights Reserved.
+
+#pragma once
+
+#include "CoreMinimal.h"
+#include "Modules/ModuleInterface.h"
+#include "Modules/ModuleManager.h"
+#include "Templates/SharedPointer.h"
+#include "Templates/UniquePtr.h"
+#include "IDerivedDataCacheNotifications.h"
+
+class FSpawnTabArgs;
+class SDerivedDataCacheSettingsDialog;
+class IDerivedDataCacheNotifications;
+class SDockTab;
+class SWidget;
+class SWindow;
+
+/**
+ * The module holding all of the UI related pieces for DerivedData
+ */
+class ZENEDITOR_API FZenEditor : public IModuleInterface
+{
+public:
+
+	/**
+	 * Called right after the module DLL has been loaded and the module object has been created
+	 */
+	virtual void StartupModule() override;
+
+	/**
+	 * Called before the module is unloaded, right before the module object is destroyed.
+	 */
+	virtual void ShutdownModule() override;
+
+	bool IsZenEnabled() const;
+	TSharedRef<SWidget>	CreateStatusBarWidget();
+
+	void ShowResourceUsageTab();
+	void ShowCacheStatisticsTab();
+	void ShowZenServerStatusTab();
+
+private:
+
+	TSharedPtr<SWidget> CreateResourceUsageDialog();
+	TSharedPtr<SWidget> CreateCacheStatisticsDialog();
+	TSharedPtr<SWidget> CreateZenStoreDialog();
+
+	TSharedRef<SDockTab> CreateResourceUsageTab(const FSpawnTabArgs& Args);
+	TSharedRef<SDockTab> CreateCacheStatisticsTab(const FSpawnTabArgs& Args);
+	TSharedRef<SDockTab> CreateZenServerStatusTab(const FSpawnTabArgs& Args);
+
+	TWeakPtr<SDockTab> ResourceUsageTab;
+	TWeakPtr<SDockTab> CacheStatisticsTab;
+	TWeakPtr<SDockTab> ZenServerStatusTab;
+
+	TSharedPtr<SWindow>	SettingsWindow;
+	TSharedPtr<SDerivedDataCacheSettingsDialog> SettingsDialog;
+	TUniquePtr<IDerivedDataCacheNotifications>	DerivedDataCacheNotifications;
+};
+
+
diff --git a/Engine/Source/Editor/ZenEditor/ZenEditor.Build.cs b/Engine/Source/Editor/ZenEditor/ZenEditor.Build.cs
new file mode 100644
--- /dev/null
+++ b/Engine/Source/Editor/ZenEditor/ZenEditor.Build.cs
@@ -0,0 +1,44 @@
+// Copyright Epic Games, Inc. All Rights Reserved.
+
+using UnrealBuildTool;
+
+public class ZenEditor : ModuleRules
+{
+	public ZenEditor(ReadOnlyTargetRules Target)
+		 : base(Target)
+	{
+		PrivateDependencyModuleNames.AddRange(
+			new string[] {
+				"Core",
+				"CoreUObject",
+				"Engine",
+				"Slate",
+				"SlateCore",
+				"InputCore",
+				"EditorFramework",
+				"UATHelper",
+				"UnrealEd",
+				"ToolMenus",
+				"OutputLog",
+				"DerivedDataCache",
+				"EditorSubsystem",
+				"WorkspaceMenuStructure",
+				"MessageLog",
+				"ToolWidgets",
+				"DerivedDataWidgets",
+				"Zen",
+				"ZenWidgets"
+			});
+
+		DynamicallyLoadedModuleNames.AddRange(
+			new string[] {
+				"MainFrame",
+			});
+
+
+		PrivateIncludePathModuleNames.AddRange(
+			new string[] {
+			});
+		CppCompileWarningSettings.UnsafeTypeCastWarningLevel = WarningLevel.Error;
+	}
+}
diff --git a/Engine/Source/Editor/ZenWidgets/Private/SMessageDialog.cpp b/Engine/Source/Editor/ZenWidgets/Private/SMessageDialog.cpp
new file mode 100644
--- /dev/null
+++ b/Engine/Source/Editor/ZenWidgets/Private/SMessageDialog.cpp
@@ -0,0 +1,447 @@
+// Copyright Epic Games, Inc. All Rights Reserved.
+
+
+#include "SMessageDialog.h"
+#include "Fonts/SlateFontInfo.h"
+#include "Framework/Application/SlateApplication.h"
+#include "HAL/PlatformApplicationMisc.h"
+#include "Input/Reply.h"
+#include "InputCoreTypes.h"
+#include "Misc/Attribute.h"
+#include "SlateOptMacros.h"
+#include "Styling/CoreStyle.h"
+#include "Widgets/DeclarativeSyntaxSupport.h"
+#include "Widgets/Input/SButton.h"
+#include "Widgets/Images/SImage.h"
+#include "Widgets/Layout/SBorder.h"
+#include "Widgets/Layout/SScrollBox.h"
+#include "Widgets/Layout/SUniformGridPanel.h"
+#include "Widgets/Input/SCheckBox.h"
+#include "Widgets/SBoxPanel.h"
+#include "Widgets/SCompoundWidget.h"
+#include "Widgets/SWindow.h"
+#include "Widgets/Text/STextBlock.h"
+
+#define LOCTEXT_NAMESPACE "Dialogs"
+
+class SChoiceDialog : public SCompoundWidget
+{
+public:
+	SLATE_BEGIN_ARGS( SChoiceDialog )	{}
+		SLATE_ARGUMENT(TSharedPtr<SWindow>, ParentWindow)
+		SLATE_ATTRIBUTE(FText, Message)	
+		SLATE_ATTRIBUTE(float, WrapMessageAt)
+		SLATE_ARGUMENT(EAppMsgCategory, MessageCategory)
+		SLATE_ARGUMENT(EAppMsgType::Type, MessageType)
+	SLATE_END_ARGS()
+
+	BEGIN_SLATE_FUNCTION_BUILD_OPTIMIZATION
+	void Construct( const FArguments& InArgs )
+	{
+		ParentWindow = InArgs._ParentWindow;
+		ParentWindow->SetWidgetToFocusOnActivate(SharedThis(this));
+		Response = EAppReturnType::Cancel;
+		MessageType = InArgs._MessageType;
+
+		FSlateFontInfo MessageFont( FAppStyle::GetFontStyle("StandardDialog.LargeFont"));
+		MyMessage = InArgs._Message;
+
+		TSharedPtr<SUniformGridPanel> ButtonBox;
+
+		const FSlateBrush* IconBrush = FAppStyle::Get().GetBrush("Icons.WarningWithColor.Large");
+		switch (InArgs._MessageCategory)
+		{
+		case EAppMsgCategory::Error:
+			IconBrush = FAppStyle::Get().GetBrush("Icons.ErrorWithColor.Large");
+			break;
+		case EAppMsgCategory::Success:
+			IconBrush = FAppStyle::Get().GetBrush("Icons.SuccessWithColor.Large");
+			break;
+		case EAppMsgCategory::Info:
+			IconBrush = FAppStyle::Get().GetBrush("Icons.InfoWithColor.Large");
+			break;
+		default:
+			break;
+		}
+
+		this->ChildSlot
+		[	
+			SNew(SBorder)
+			.BorderImage(FAppStyle::GetBrush("ToolPanel.GroupBorder"))
+			.Padding(16.f)
+			[
+				SNew(SVerticalBox)
+
+				+ SVerticalBox::Slot()
+				.HAlign(HAlign_Fill)
+				.VAlign(VAlign_Fill)
+				.FillHeight(1.0f)
+				.MaxHeight(550)
+				[
+					SNew(SHorizontalBox)
+					+SHorizontalBox::Slot()
+					.AutoWidth()
+					.VAlign(VAlign_Top)
+					.HAlign(HAlign_Left)
+					[
+						SNew(SImage)
+						.DesiredSizeOverride(FVector2D(24.f, 24.f))
+						.Image(IconBrush)
+					]
+
+					+SHorizontalBox::Slot()
+					.Padding(16.f, 0.f, 0.f, 0.f)
+					[
+						SNew(SScrollBox)
+						+ SScrollBox::Slot()
+						.HAlign(HAlign_Left)
+						.VAlign(VAlign_Center)
+						[
+							SNew(STextBlock)
+							.Text(MyMessage)
+							.Font(MessageFont)
+							.WrapTextAt(InArgs._WrapMessageAt)
+						]
+					]
+				]
+
+				+SVerticalBox::Slot()
+				.Padding(FMargin(0.f, 32.f, 0.f, 0.f))
+				.AutoHeight()
+				[
+					SNew(SHorizontalBox)
+
+					+ SHorizontalBox::Slot()
+					.AutoWidth()
+					.HAlign(HAlign_Left)
+					.VAlign(VAlign_Center)
+					[
+						SNew(SButton)
+						.ButtonStyle(FAppStyle::Get(), "SimpleButton")
+						.OnClicked(this, &SChoiceDialog::HandleCopyMessageButtonClicked)
+						.ToolTipText(NSLOCTEXT("SChoiceDialog", "CopyMessageTooltip", "Copy the text in this message to the clipboard (CTRL+C)"))
+						.ContentPadding(2.f)
+						.Content()
+						[
+							SNew(SImage)
+							.Image(FAppStyle::Get().GetBrush("Zen.Icons.Clipboard"))
+							.ColorAndOpacity(FSlateColor::UseForeground())
+						]
+					]
+
+					+ SHorizontalBox::Slot()
+						.FillWidth(1.0f)
+						.HAlign(HAlign_Left)
+						.VAlign(VAlign_Center)
+						.Padding(FMargin(16.f, 0.f, 0.f, 0.f))
+						[
+							SNew(SCheckBox)
+							.IsChecked(ECheckBoxState::Unchecked)
+							.OnCheckStateChanged(this, &SChoiceDialog::OnCheckboxClicked)
+							.Visibility(this, &SChoiceDialog::GetCheckboxVisibility)
+							.ToolTipText(NSLOCTEXT("SChoiceDialog", "ApplyToAllTooltip", "Make your choice of Yes or No apply to all remaining items in the current operation"))
+							[
+								SNew(STextBlock)
+								.WrapTextAt(615.0f)
+								.Text(NSLOCTEXT("SChoiceDialog", "ApplyToAllLabel", "Apply to All"))
+							]
+						]
+
+					+ SHorizontalBox::Slot()
+					.AutoWidth()
+					.HAlign(HAlign_Right)
+					.VAlign(VAlign_Center)
+					[
+						SAssignNew( ButtonBox, SUniformGridPanel )
+						.SlotPadding(FMargin(16.f, 0.f, 0.f, 0.f))
+						.MinDesiredSlotWidth(FAppStyle::Get().GetFloat("StandardDialog.MinDesiredSlotWidth"))
+						.MinDesiredSlotHeight(FAppStyle::Get().GetFloat("StandardDialog.MinDesiredSlotHeight"))
+					]
+				]
+			]
+		];
+
+		int32 SlotIndex = 0;
+
+#define ADD_SLOT(Button)\
+		ButtonBox->AddSlot(SlotIndex++,0)\
+		[\
+			SNew( SButton )\
+			.VAlign(VAlign_Center)\
+			.HAlign(HAlign_Center)\
+			.Text( EAppReturnTypeToText(EAppReturnType::Button) )\
+			.OnClicked( this, &SChoiceDialog::HandleButtonClicked, EAppReturnType::Button )\
+		];
+
+#define ADD_SLOT_PRIMARY(Button)\
+		ButtonBox->AddSlot(SlotIndex++,0)\
+		[\
+			SNew( SButton )\
+			.VAlign(VAlign_Center)\
+			.HAlign(HAlign_Center)\
+			.ButtonStyle(&FAppStyle::Get().GetWidgetStyle<FButtonStyle>("PrimaryButton"))\
+			.Text( EAppReturnTypeToText(EAppReturnType::Button) )\
+			.OnClicked( this, &SChoiceDialog::HandleButtonClicked, EAppReturnType::Button )\
+		];
+
+
+
+		switch ( MessageType )
+		{	
+		case EAppMsgType::Ok:
+			ADD_SLOT_PRIMARY(Ok)
+			break;
+		case EAppMsgType::YesNo:
+			ADD_SLOT_PRIMARY(Yes)
+			ADD_SLOT(No)
+			break;
+		case EAppMsgType::OkCancel:
+			ADD_SLOT_PRIMARY(Ok)
+			ADD_SLOT(Cancel)
+			break;
+		case EAppMsgType::YesNoCancel:
+			ADD_SLOT_PRIMARY(Yes)
+			ADD_SLOT(No)
+			ADD_SLOT(Cancel)
+			break;
+		case EAppMsgType::CancelRetryContinue:
+			ADD_SLOT_PRIMARY(Continue)
+			ADD_SLOT(Retry)
+			ADD_SLOT(Cancel)
+			break;
+		case EAppMsgType::YesNoYesAllNoAll:
+			ADD_SLOT_PRIMARY(Yes)
+			ADD_SLOT(No)
+			break;
+		case EAppMsgType::YesNoYesAllNoAllCancel:
+			ADD_SLOT_PRIMARY(Yes)
+			ADD_SLOT(No)
+			ADD_SLOT(Cancel)
+			break;
+		case EAppMsgType::YesNoYesAll:
+			ADD_SLOT_PRIMARY(Yes)
+			ADD_SLOT(No)
+			ADD_SLOT(YesAll)
+			break;
+		}
+
+#undef ADD_SLOT
+#undef ADD_SLOT_PRIMARY
+	}
+	END_SLATE_FUNCTION_BUILD_OPTIMIZATION
+
+
+	EAppReturnType::Type GetResponse()
+	{
+		return Response;
+	}
+
+	virtual	FReply OnKeyDown( const FGeometry& MyGeometry, const FKeyEvent& InKeyEvent )
+	{
+		if( InKeyEvent.GetKey() == EKeys::Escape )
+		{
+			return HandleButtonClicked(EAppReturnType::Cancel);
+		}
+
+		if (InKeyEvent.GetKey() == EKeys::C && InKeyEvent.IsControlDown())
+		{
+			CopyMessageToClipboard();
+
+			return FReply::Handled();
+		}
+
+		//if it was some other button, ignore it
+		return FReply::Unhandled();
+	}
+
+	/** Override the base method to allow for keyboard focus */
+	virtual bool SupportsKeyboardFocus() const
+	{
+		return true;
+	}
+
+	/** Converts an EAppReturnType into a localized FText */
+	static FText EAppReturnTypeToText(EAppReturnType::Type ReturnType)
+	{
+		switch(ReturnType)
+		{
+		case EAppReturnType::No:
+			return LOCTEXT("EAppReturnTypeNo", "No");
+		case EAppReturnType::Yes:
+			return LOCTEXT("EAppReturnTypeYes", "Yes");
+		case EAppReturnType::YesAll:
+			return LOCTEXT("EAppReturnTypeYesAll", "Yes All");
+		case EAppReturnType::NoAll:
+			return LOCTEXT("EAppReturnTypeNoAll", "No All");
+		case EAppReturnType::Cancel:
+			return LOCTEXT("EAppReturnTypeCancel", "Cancel");
+		case EAppReturnType::Ok:
+			return LOCTEXT("EAppReturnTypeOk", "OK");
+		case EAppReturnType::Retry:
+			return LOCTEXT("EAppReturnTypeRetry", "Retry");
+		case EAppReturnType::Continue:
+			return LOCTEXT("EAppReturnTypeContinue", "Continue");
+		default:
+			return LOCTEXT("MissingType", "MISSING RETURN TYPE");
+		}
+	}
+
+protected:
+
+	/**
+	 * Copies the message text to the clipboard.
+	 */
+	void CopyMessageToClipboard( )
+	{
+		FPlatformApplicationMisc::ClipboardCopy( *MyMessage.Get().ToString() );
+	}
+
+
+private:
+
+	// Handles clicking a message box button.
+	FReply HandleButtonClicked( EAppReturnType::Type InResponse )
+	{
+		Response = InResponse;
+		if ((MessageType == EAppMsgType::YesNoYesAllNoAll || MessageType == EAppMsgType::YesNoYesAllNoAllCancel)
+			&& bApplyToAllChecked)
+		{
+			if (Response == EAppReturnType::Yes)
+			{
+				Response = EAppReturnType::YesAll;
+			}
+			else if (Response == EAppReturnType::No)
+			{
+				Response = EAppReturnType::NoAll;
+			}
+		}
+
+		ResultCallback.ExecuteIfBound(ParentWindow.ToSharedRef(), Response);
+
+
+		ParentWindow->RequestDestroyWindow();
+
+		return FReply::Handled();
+	}
+
+	// Handles clicking the 'Copy Message' button.
+	FReply HandleCopyMessageButtonClicked()
+	{
+		CopyMessageToClipboard();
+		return FReply::Handled();
+	}
+		
+	// Used as a delegate for the OnClicked property of the Apply to All checkbox
+	void OnCheckboxClicked(ECheckBoxState InNewState)
+	{
+		bApplyToAllChecked = InNewState == ECheckBoxState::Checked;
+	}
+
+	// Used as a delegate for the Visibility property of the Apply to All checkbox
+	EVisibility GetCheckboxVisibility() const
+	{
+		if (MessageType == EAppMsgType::YesNoYesAllNoAll || MessageType == EAppMsgType::YesNoYesAllNoAllCancel)
+		{
+			return EVisibility::Visible;
+		}
+		return EVisibility::Hidden;
+	}
+
+public:
+	/** Callback delegate that is triggered, when the dialog is run in non-modal mode */
+	FOnMsgDlgResult ResultCallback;
+
+private:
+
+	EAppReturnType::Type Response;
+	TSharedPtr<SWindow> ParentWindow;
+	TAttribute<FText> MyMessage;
+	bool bApplyToAllChecked;
+	EAppMsgType::Type MessageType;
+};
+
+
+void CreateMsgDlgWindow(TSharedPtr<SWindow>& OutWindow, TSharedPtr<SChoiceDialog>& OutDialog, EAppMsgCategory InMessageCategory, EAppMsgType::Type InMessageType,
+						const FText& InMessage, const FText& InTitle, FOnMsgDlgResult ResultCallback=NULL)
+{
+	OutWindow = SNew(SWindow)
+		.Title(InTitle)
+		.SizingRule(ESizingRule::Autosized)
+		.AutoCenter(EAutoCenter::PreferredWorkArea)
+		.SupportsMinimize(false).SupportsMaximize(false);
+
+	OutDialog = SNew(SChoiceDialog)
+		.ParentWindow(OutWindow)
+		.Message(InMessage)
+		.WrapMessageAt(512.0f)
+		.MessageCategory(InMessageCategory)
+		.MessageType(InMessageType);
+
+	OutDialog->ResultCallback = ResultCallback;
+
+	OutWindow->SetContent(OutDialog.ToSharedRef());
+}
+
+EAppReturnType::Type OpenModalMessageDialog_Internal(EAppMsgCategory InMessageCategory, EAppMsgType::Type InMessageType, EAppReturnType::Type InDefaultValue, const FText& InMessage, const FText& InTitle, const TSharedPtr<const SWidget>& ModalParent)
+{
+	EAppReturnType::Type Response = InDefaultValue;
+	TSharedPtr<SWindow> MsgWindow = NULL;
+	TSharedPtr<SChoiceDialog> MsgDialog = NULL;
+
+	CreateMsgDlgWindow(MsgWindow, MsgDialog, InMessageCategory, InMessageType, InMessage, InTitle);
+
+	FSlateApplication::Get().AddModalWindow(MsgWindow.ToSharedRef(), ModalParent);
+
+	Response = MsgDialog->GetResponse();
+
+	return Response;
+}
+
+EAppReturnType::Type OpenModalMessageDialog_Internal(EAppMsgCategory InMessageCategory, EAppMsgType::Type InMessageType, const FText& InMessage, const FText& InTitle, const TSharedPtr<const SWidget>& ModalParent)
+{
+	EAppReturnType::Type DefaultValue = EAppReturnType::Yes;
+	switch (InMessageType)
+	{
+	case EAppMsgType::Ok:
+		DefaultValue = EAppReturnType::Ok;
+		break;
+	case EAppMsgType::YesNo:
+		DefaultValue = EAppReturnType::No;
+		break;
+	case EAppMsgType::OkCancel:
+		DefaultValue = EAppReturnType::Cancel;
+		break;
+	case EAppMsgType::YesNoCancel:
+		DefaultValue = EAppReturnType::Cancel;
+		break;
+	case EAppMsgType::CancelRetryContinue:
+		DefaultValue = EAppReturnType::Cancel;
+		break;
+	case EAppMsgType::YesNoYesAllNoAll:
+		DefaultValue = EAppReturnType::No;
+		break;
+	case EAppMsgType::YesNoYesAllNoAllCancel:
+		DefaultValue = EAppReturnType::Cancel;
+		break;
+	default:
+		DefaultValue = EAppReturnType::Yes;
+		break;
+	}
+
+	if (GIsRunningUnattendedScript && InMessageType != EAppMsgType::Ok)
+	{
+		//UE_LOG(LogDialogs, Error, TEXT("Message Dialog was triggered in unattended script mode without a default value. %d will be used."), (int32)DefaultValue);
+		if (FPlatformMisc::IsDebuggerPresent())
+		{
+			UE_DEBUG_BREAK();
+		}
+		else
+		{
+			FDebug::DumpStackTraceToLog(ELogVerbosity::Error);
+		}
+	}
+
+	return OpenModalMessageDialog_Internal(InMessageCategory, InMessageType, DefaultValue, InMessage, InTitle, ModalParent);
+}
+
+#undef LOCTEXT_NAMESPACE
diff --git a/Engine/Source/Editor/ZenWidgets/Private/SZenCacheStatistics.cpp b/Engine/Source/Editor/ZenWidgets/Private/SZenCacheStatistics.cpp
new file mode 100644
--- /dev/null
+++ b/Engine/Source/Editor/ZenWidgets/Private/SZenCacheStatistics.cpp
@@ -0,0 +1,313 @@
+// Copyright Epic Games, Inc. All Rights Reserved.
+
+#include "SZenCacheStatistics.h"
+#include "Experimental/ZenServerInterface.h"
+#include "Internationalization/FastDecimalFormat.h"
+#include "Math/BasicMathExpressionEvaluator.h"
+#include "Math/UnitConversion.h"
+#include "Misc/ExpressionParser.h"
+#include "Styling/StyleColors.h"
+#include "Widgets/Layout/SGridPanel.h"
+#include "Widgets/SBoxPanel.h"
+#include "Widgets/Text/STextBlock.h"
+
+#define LOCTEXT_NAMESPACE "ZenDashboard"
+
+void SZenCacheStatistics::Construct(const FArguments& InArgs)
+{
+	ZenServiceInstance = InArgs._ZenServiceInstance;
+
+	this->ChildSlot
+	[
+		SNew(SVerticalBox)
+
+		+ SVerticalBox::Slot()
+		.AutoHeight()
+		.Padding(0, 0, 0, 0)
+		.Expose(GridSlot)
+		[
+			GetGridPanel()
+		]
+	];
+
+	RegisterActiveTimer(0.5f, FWidgetActiveTimerDelegate::CreateSP(this, &SZenCacheStatistics::UpdateGridPanels));
+}
+
+EActiveTimerReturnType SZenCacheStatistics::UpdateGridPanels(double InCurrentTime, float InDeltaTime)
+{
+	(*GridSlot)
+	[
+		GetGridPanel()
+	];
+
+	SlatePrepass(GetPrepassLayoutScaleMultiplier());
+
+	return EActiveTimerReturnType::Continue;
+}
+
+TSharedRef<SWidget> SZenCacheStatistics::GetGridPanel()
+{
+	TSharedRef<SGridPanel> Panel = SNew(SGridPanel);
+
+	bool bHaveStats = false;
+	UE::Zen::FZenCacheStats ZenStats;
+
+	if (TSharedPtr<UE::Zen::FZenServiceInstance> ServiceInstance = ZenServiceInstance.Get())
+	{
+		if (ServiceInstance->GetCacheStats(ZenStats))
+		{
+			bHaveStats = true;
+		}
+	}
+
+	const static FNumberFormattingOptions SingleDecimalFormatting = FNumberFormattingOptions()
+		.SetUseGrouping(true)
+		.SetMinimumFractionalDigits(1)
+		.SetMaximumFractionalDigits(1);
+
+	int32 Row = 0;
+
+	const float RowMargin = 0.0f;
+	const float TitleMargin = 10.0f;
+	const float ColumnMargin = 10.0f;
+	const FSlateColor TitleColor = FStyleColors::AccentWhite;
+	const FSlateFontInfo TitleFont = FCoreStyle::GetDefaultFontStyle("Bold", 10);
+
+
+	Panel->AddSlot(0, Row)
+	.ColumnSpan(2)
+	[
+		SNew(STextBlock)
+		.Margin(FMargin(ColumnMargin, RowMargin))
+		.ColorAndOpacity(TitleColor)
+		.Font(TitleFont)
+		.Text(LOCTEXT("Cache", "Local Cache"))
+	];
+
+	Row++;
+
+	Panel->AddSlot(0, Row)
+	[
+		SNew(STextBlock)
+		.Margin(FMargin(ColumnMargin, RowMargin))
+		.Text(LOCTEXT("CacheHitRate", "Hit Rate"))
+	];
+
+	Panel->AddSlot(1, Row)
+	[
+		SNew(STextBlock)
+		.Margin(FMargin(ColumnMargin, RowMargin))
+		.Text_Lambda([CacheStats = ZenStats.General]
+		{
+			int64 CacheTotal = CacheStats.Hits + CacheStats.Misses;
+			return CacheTotal == 0 ? LOCTEXT("CacheNoHitRateValue", "-") : FText::AsPercent(CacheStats.HitRatio, &SingleDecimalFormatting);
+		})
+	];
+
+	Row++;
+
+	Panel->AddSlot(0, Row)
+	[
+		SNew(STextBlock)
+		.Margin(FMargin(ColumnMargin, RowMargin))
+		.Text(LOCTEXT("CacheHitQuantity", "Hits"))
+	];
+
+	Panel->AddSlot(1, Row)
+	[
+		SNew(STextBlock)
+		.Margin(FMargin(ColumnMargin, RowMargin))
+		.Text_Lambda([bHaveStats, CacheHits = ZenStats.General.Hits]
+		{
+			return bHaveStats ? FText::AsNumber(CacheHits) : LOCTEXT("UnavailableValue", "-");
+		})
+	];
+
+	Row++;
+
+	Panel->AddSlot(0, Row)
+	[
+		SNew(STextBlock)
+		.Margin(FMargin(ColumnMargin, RowMargin))
+		.Text(LOCTEXT("CacheMissQuantity", "Misses"))
+	];
+
+	Panel->AddSlot(1, Row)
+	[
+		SNew(STextBlock)
+		.Margin(FMargin(ColumnMargin, RowMargin))
+		.Text_Lambda([bHaveStats, CacheMisses = ZenStats.General.Misses]
+		{
+			return bHaveStats ? FText::AsNumber(CacheMisses) : LOCTEXT("UnavailableValue", "-");
+		})
+	];
+
+	Row++;
+
+	Panel->AddSlot(0, Row)
+	[
+		SNew(STextBlock)
+		.Margin(FMargin(ColumnMargin, RowMargin))
+		.Text(LOCTEXT("CacheWriteQuantity", "Writes"))
+	];
+
+	Panel->AddSlot(1, Row)
+	[
+		SNew(STextBlock)
+		.Margin(FMargin(ColumnMargin, RowMargin))
+		.Text_Lambda([bHaveStats, CacheWrites = ZenStats.General.Writes]
+		{
+			return bHaveStats ? FText::AsNumber(CacheWrites) : LOCTEXT("UnavailableValue", "-");
+		})
+	];
+
+	Row++;
+
+	Panel->AddSlot(0, Row)
+	[
+		SNew(STextBlock)
+		.Margin(FMargin(ColumnMargin, RowMargin))
+		.Text(LOCTEXT("CacheRequests", "Requests"))
+	];
+
+	Panel->AddSlot(1, Row)
+	[
+		SNew(STextBlock)
+		.Margin(FMargin(ColumnMargin, RowMargin))
+		.Text_Lambda([bHaveStats, Requests = ZenStats.Request.Count]
+		{
+			if (bHaveStats)
+			{
+				return FText::AsNumber(Requests);
+			}
+			return LOCTEXT("UnavailableValue", "-");
+		})
+	];
+
+	Row++;
+
+	Panel->AddSlot(0, Row)
+	[
+		SNew(STextBlock)
+		.Margin(FMargin(ColumnMargin, RowMargin))
+		.Text(LOCTEXT("CacheBadRequests", "Bad Requests"))
+	];
+
+	Panel->AddSlot(1, Row)
+	[
+		SNew(STextBlock)
+		.Margin(FMargin(ColumnMargin, RowMargin))
+		.Text_Lambda([bHaveStats, BadRequests = ZenStats.General.BadRequestCount]
+		{
+			if (bHaveStats)
+			{
+				return FText::AsNumber(BadRequests);
+			}
+			return LOCTEXT("UnavailableValue", "-");
+		})
+	];
+
+	Row++;
+
+	if (!ZenStats.Upstream.EndPoint.IsEmpty())
+	{
+		Panel->AddSlot(0, Row)
+		[
+			SNew(STextBlock)
+			.Margin(FMargin(ColumnMargin, TitleMargin, ColumnMargin, 0.0f))
+			.ColorAndOpacity(TitleColor)
+			.Font(TitleFont)
+			.Text(LOCTEXT("UpstreamServer", "Upstream Server"))
+		];
+
+		Panel->AddSlot(1, Row)
+		[
+			SNew(STextBlock)
+			.Margin(FMargin(ColumnMargin, TitleMargin, ColumnMargin, 0.0f))
+			.ColorAndOpacity(TitleColor)
+			.Font(TitleFont)
+			.Text(LOCTEXT("HitRate", "Hit Rate"))
+		];
+
+		Panel->AddSlot(2, Row)
+		[
+			SNew(STextBlock)
+			.Margin(FMargin(ColumnMargin, TitleMargin, ColumnMargin, 0.0f))
+			.ColorAndOpacity(TitleColor)
+			.Font(TitleFont)
+			.Text(LOCTEXT("Downloaded", "Downloaded"))
+		];
+
+		Panel->AddSlot(3, Row)
+		[
+			SNew(STextBlock)
+			.Margin(FMargin(ColumnMargin, TitleMargin, ColumnMargin, 0.0f))
+			.ColorAndOpacity(TitleColor)
+			.Font(TitleFont)
+			.Text(LOCTEXT("Uploaded", "Uploaded"))
+		];
+
+		Panel->AddSlot(4, Row)
+		[
+			SNew(STextBlock)
+			.Margin(FMargin(ColumnMargin, TitleMargin, ColumnMargin, 0.0f))
+			.ColorAndOpacity(TitleColor)
+			.Font(TitleFont)
+			.Text(LOCTEXT("URL", "URL"))
+		];
+		Row++;
+
+		for (const UE::Zen::FZenCacheStats::FEndPointStats& EndpointStats : ZenStats.Upstream.EndPoint)
+		{
+			Panel->AddSlot(0, Row)
+			[
+				SNew(STextBlock)
+				.Margin(FMargin(ColumnMargin, RowMargin))
+				.Text_Lambda([EndpointStats] { return FText::FromString(EndpointStats.Name); })
+			];
+
+			Panel->AddSlot(1, Row)
+			[
+				SNew(STextBlock)
+				.Margin(FMargin(ColumnMargin, RowMargin))
+				.Text_Lambda([EndpointStats] { return FText::AsPercent(EndpointStats.HitRatio, &SingleDecimalFormatting); })
+			];
+
+			Panel->AddSlot(2, Row)
+			[
+				SNew(STextBlock)
+				.Margin(FMargin(ColumnMargin, RowMargin))
+				.Text_Lambda([DownloadedMB = EndpointStats.DownloadedMB]
+					{
+						uint64 DownloadedBytes = uint64(DownloadedMB*1024*1024);
+						return FText::AsMemory(DownloadedBytes, (DownloadedBytes > 1024) ? &SingleDecimalFormatting : nullptr, nullptr, EMemoryUnitStandard::IEC);
+					})
+												
+			];
+
+			Panel->AddSlot(3, Row)
+			[
+				SNew(STextBlock)
+				.Margin(FMargin(ColumnMargin, RowMargin))
+				.Text_Lambda([UploadedMB = EndpointStats.UploadedMB]
+					{
+						uint64 UploadedBytes = uint64(UploadedMB*1024*1024);
+						return FText::AsMemory(UploadedBytes, (UploadedBytes > 1024) ? &SingleDecimalFormatting : nullptr, nullptr, EMemoryUnitStandard::IEC);
+					})
+			];
+
+			Panel->AddSlot(4, Row)
+			[
+				SNew(STextBlock)
+				.Margin(FMargin(ColumnMargin, RowMargin))
+				.Text_Lambda([EndpointStats] { return FText::FromString(EndpointStats.Url); })
+			];
+
+			Row++;
+		}
+	}
+
+	return Panel;
+}
+
+#undef LOCTEXT_NAMESPACE
diff --git a/Engine/Source/Editor/ZenWidgets/Private/SZenCidStoreStatistics.cpp b/Engine/Source/Editor/ZenWidgets/Private/SZenCidStoreStatistics.cpp
new file mode 100644
--- /dev/null
+++ b/Engine/Source/Editor/ZenWidgets/Private/SZenCidStoreStatistics.cpp
@@ -0,0 +1,213 @@
+// Copyright Epic Games, Inc. All Rights Reserved.
+
+#include "SZenCidStoreStatistics.h"
+#include "Experimental/ZenServerInterface.h"
+#include "Internationalization/FastDecimalFormat.h"
+#include "Math/BasicMathExpressionEvaluator.h"
+#include "Math/UnitConversion.h"
+#include "Misc/ExpressionParser.h"
+#include "Styling/StyleColors.h"
+#include "Widgets/Layout/SGridPanel.h"
+#include "Widgets/SBoxPanel.h"
+#include "Widgets/Text/STextBlock.h"
+
+#define LOCTEXT_NAMESPACE "ZenDashboard"
+
+void SZenCidStoreStatistics::Construct(const FArguments& InArgs)
+{
+	ZenServiceInstance = InArgs._ZenServiceInstance;
+
+	this->ChildSlot
+	[
+		SNew(SVerticalBox)
+
+		+ SVerticalBox::Slot()
+		.AutoHeight()
+		.Padding(0, 0, 0, 0)
+		.Expose(GridSlot)
+		[
+			GetGridPanel()
+		]
+	];
+
+	RegisterActiveTimer(0.5f, FWidgetActiveTimerDelegate::CreateSP(this, &SZenCidStoreStatistics::UpdateGridPanels));
+}
+
+EActiveTimerReturnType SZenCidStoreStatistics::UpdateGridPanels(double InCurrentTime, float InDeltaTime)
+{
+	(*GridSlot)
+	[
+		GetGridPanel()
+	];
+
+	SlatePrepass(GetPrepassLayoutScaleMultiplier());
+
+	return EActiveTimerReturnType::Continue;
+}
+
+TSharedRef<SWidget> SZenCidStoreStatistics::GetGridPanel()
+{
+	TSharedRef<SGridPanel> Panel = SNew(SGridPanel);
+
+	bool bHaveStats = false;
+	UE::Zen::FZenCacheStats ZenStats;
+
+	if (TSharedPtr<UE::Zen::FZenServiceInstance> ServiceInstance = ZenServiceInstance.Get())
+	{
+		if (ServiceInstance->GetCacheStats(ZenStats))
+		{
+			bHaveStats = true;
+		}
+	}
+
+	int64 TotalHits = ZenStats.General.CidHits;
+	int64 TotalMisses = ZenStats.General.CidMisses;
+	int64 TotalWrites = ZenStats.General.CidWrites;
+	int64 Requests = ZenStats.General.CidHits + ZenStats.General.CidMisses + ZenStats.General.CidWrites;
+
+	const static FNumberFormattingOptions SingleDecimalFormatting = FNumberFormattingOptions()
+		.SetUseGrouping(true)
+		.SetMinimumFractionalDigits(1)
+		.SetMaximumFractionalDigits(1);
+
+	int32 Row = 0;
+
+	const float RowMargin = 0.0f;
+	const float TitleMargin = 10.0f;
+	const float ColumnMargin = 10.0f;
+	const FSlateColor TitleColor = FStyleColors::AccentWhite;
+	const FSlateFontInfo TitleFont = FCoreStyle::GetDefaultFontStyle("Bold", 10);
+
+
+	Panel->AddSlot(0, Row)
+	.ColumnSpan(2)
+	[
+		SNew(STextBlock)
+		.Margin(FMargin(ColumnMargin, RowMargin))
+		.ColorAndOpacity(TitleColor)
+		.Font(TitleFont)
+		.Text(LOCTEXT("CASCid", "Local Cid Store"))
+	];
+
+	Row++;
+
+	Panel->AddSlot(0, Row)
+	[
+		SNew(STextBlock)
+		.Margin(FMargin(ColumnMargin, RowMargin))
+		.Text(LOCTEXT("CidStoreHitRate", "Hit Rate"))
+	];
+
+	Panel->AddSlot(1, Row)
+	[
+		SNew(STextBlock)
+		.Margin(FMargin(ColumnMargin, RowMargin))
+		.Text_Lambda([bHaveStats, TotalHits, TotalMisses]
+		{
+			if (bHaveStats && ((TotalHits+TotalMisses) > 0))
+			{
+				double HitRate = static_cast<double>(TotalHits) / static_cast<double>(TotalHits + TotalMisses);
+				return FText::AsPercent(HitRate, &SingleDecimalFormatting);
+			}
+			return LOCTEXT("UnavailableValue", "-");
+		})
+	];
+
+	Row++;
+
+	Panel->AddSlot(0, Row)
+	[
+		SNew(STextBlock)
+		.Margin(FMargin(ColumnMargin, RowMargin))
+		.Text(LOCTEXT("CidStoreHitQuantity", "Hits"))
+	];
+
+	Panel->AddSlot(1, Row)
+	[
+		SNew(STextBlock)
+		.Margin(FMargin(ColumnMargin, RowMargin))
+		.Text_Lambda([bHaveStats, TotalHits]
+		{
+			if (bHaveStats)
+			{
+				return FText::AsNumber(TotalHits);
+			}
+			return LOCTEXT("UnavailableValue", "-");
+		})
+	];
+
+	Row++;
+
+	Panel->AddSlot(0, Row)
+	[
+		SNew(STextBlock)
+		.Margin(FMargin(ColumnMargin, RowMargin))
+		.Text(LOCTEXT("CidStoreMissQuantity", "Misses"))
+	];
+
+	Panel->AddSlot(1, Row)
+	[
+		SNew(STextBlock)
+		.Margin(FMargin(ColumnMargin, RowMargin))
+		.Text_Lambda([bHaveStats, TotalMisses]
+		{
+			if (bHaveStats)
+			{
+				return FText::AsNumber(TotalMisses);
+			}
+			return LOCTEXT("UnavailableValue", "-");
+		})
+	];
+
+	Row++;
+
+	Panel->AddSlot(0, Row)
+	[
+		SNew(STextBlock)
+		.Margin(FMargin(ColumnMargin, RowMargin))
+		.Text(LOCTEXT("CidStoreWriteQuantity", "Writes"))
+	];
+
+	Panel->AddSlot(1, Row)
+	[
+		SNew(STextBlock)
+		.Margin(FMargin(ColumnMargin, RowMargin))
+		.Text_Lambda([bHaveStats, TotalWrites]
+		{
+			if (bHaveStats)
+			{
+				return FText::AsNumber(TotalWrites);
+			}
+			return LOCTEXT("UnavailableValue", "-");
+		})
+	];
+
+	Row++;
+
+	Panel->AddSlot(0, Row)
+	[
+		SNew(STextBlock)
+		.Margin(FMargin(ColumnMargin, RowMargin))
+		.Text(LOCTEXT("ProjectRequests", "Requests"))
+	];
+
+	Panel->AddSlot(1, Row)
+	[
+		SNew(STextBlock)
+		.Margin(FMargin(ColumnMargin, RowMargin))
+		.Text_Lambda([bHaveStats, Requests]
+		{
+			if (bHaveStats)
+			{
+				return FText::AsNumber(Requests);
+			}
+			return LOCTEXT("UnavailableValue", "-");
+		})
+	];
+
+	Row++;
+
+	return Panel;
+}
+
+#undef LOCTEXT_NAMESPACE
diff --git a/Engine/Source/Editor/ZenWidgets/Private/SZenDialogs.cpp b/Engine/Source/Editor/ZenWidgets/Private/SZenDialogs.cpp
new file mode 100644
--- /dev/null
+++ b/Engine/Source/Editor/ZenWidgets/Private/SZenDialogs.cpp
@@ -0,0 +1,75 @@
+// Copyright Epic Games, Inc. All Rights Reserved.
+
+#include "SZenDialogs.h"
+
+#include "ZenServiceInstanceManager.h"
+#include "SMessageDialog.h"
+#include "SZenCacheStatistics.h"
+#include "SZenCidStoreStatistics.h"
+#include "SZenProjectStatistics.h"
+#include "SZenServiceStatus.h"
+
+#include "Framework/Application/SlateApplication.h"
+#include "Styling/StyleColors.h"
+#include "Widgets/Images/SImage.h"
+#include "Widgets/Layout/SGridPanel.h"
+#include "Widgets/SBoxPanel.h"
+#include "Widgets/SToolTip.h"
+#include "Widgets/Text/STextBlock.h"
+
+#define LOCTEXT_NAMESPACE "ZenEditor"
+
+void SZenStoreStausDialog::Construct(const FArguments& InArgs)
+{
+	TSharedPtr<UE::Zen::FServiceInstanceManager> ServiceInstanceManager = MakeShared<UE::Zen::FServiceInstanceManager>();
+
+	this->ChildSlot
+		[
+			SNew(SVerticalBox)
+				+ SVerticalBox::Slot()
+				.AutoHeight()
+				.Padding(0.0f, 10.0f, 0.0f, 0.0f)
+				.VAlign(VAlign_Top)
+				[
+					SNew(SZenServiceStatus)
+						.ZenServiceInstance(ServiceInstanceManager.ToSharedRef(), &UE::Zen::FServiceInstanceManager::GetZenServiceInstance)
+				]
+				// Stats panel
+				+ SVerticalBox::Slot()
+				.AutoHeight()
+				.Padding(0.0f, 10.0f, 0.0f, 0.0f)
+				.VAlign(VAlign_Top)
+				[
+					SNew(SHorizontalBox)
+
+						+ SHorizontalBox::Slot()
+						.AutoWidth()
+						.HAlign(HAlign_Left)
+						.VAlign(VAlign_Top)
+						[
+							SNew(SZenCacheStatistics)
+								.ZenServiceInstance(ServiceInstanceManager.ToSharedRef(), &UE::Zen::FServiceInstanceManager::GetZenServiceInstance)
+						]
+
+						+ SHorizontalBox::Slot()
+						.AutoWidth()
+						.HAlign(HAlign_Left)
+						.VAlign(VAlign_Top)
+						[
+							SNew(SZenProjectStatistics)
+								.ZenServiceInstance(ServiceInstanceManager.ToSharedRef(), &UE::Zen::FServiceInstanceManager::GetZenServiceInstance)
+						]
+
+						+ SHorizontalBox::Slot()
+						.AutoWidth()
+						.HAlign(HAlign_Left)
+						.VAlign(VAlign_Top)
+						[
+							SNew(SZenCidStoreStatistics)
+								.ZenServiceInstance(ServiceInstanceManager.ToSharedRef(), &UE::Zen::FServiceInstanceManager::GetZenServiceInstance)
+						]
+				]
+		];
+}
+
+#undef LOCTEXT_NAMESPACE
diff --git a/Engine/Source/Editor/ZenWidgets/Private/SZenProjectStatistics.cpp b/Engine/Source/Editor/ZenWidgets/Private/SZenProjectStatistics.cpp
new file mode 100644
--- /dev/null
+++ b/Engine/Source/Editor/ZenWidgets/Private/SZenProjectStatistics.cpp
@@ -0,0 +1,235 @@
+// Copyright Epic Games, Inc. All Rights Reserved.
+
+#include "SZenProjectStatistics.h"
+#include "Experimental/ZenServerInterface.h"
+#include "Internationalization/FastDecimalFormat.h"
+#include "Math/BasicMathExpressionEvaluator.h"
+#include "Math/UnitConversion.h"
+#include "Misc/ExpressionParser.h"
+#include "Styling/StyleColors.h"
+#include "Widgets/Layout/SGridPanel.h"
+#include "Widgets/SBoxPanel.h"
+#include "Widgets/Text/STextBlock.h"
+
+#define LOCTEXT_NAMESPACE "ZenDashboard"
+
+void SZenProjectStatistics::Construct(const FArguments& InArgs)
+{
+	ZenServiceInstance = InArgs._ZenServiceInstance;
+
+	this->ChildSlot
+	[
+		SNew(SVerticalBox)
+
+		+ SVerticalBox::Slot()
+		.AutoHeight()
+		.Padding(0, 0, 0, 0)
+		.Expose(GridSlot)
+		[
+			GetGridPanel()
+		]
+	];
+
+	RegisterActiveTimer(0.5f, FWidgetActiveTimerDelegate::CreateSP(this, &SZenProjectStatistics::UpdateGridPanels));
+}
+
+EActiveTimerReturnType SZenProjectStatistics::UpdateGridPanels(double InCurrentTime, float InDeltaTime)
+{
+	(*GridSlot)
+	[
+		GetGridPanel()
+	];
+
+	SlatePrepass(GetPrepassLayoutScaleMultiplier());
+
+	return EActiveTimerReturnType::Continue;
+}
+
+TSharedRef<SWidget> SZenProjectStatistics::GetGridPanel()
+{
+	TSharedRef<SGridPanel> Panel = SNew(SGridPanel);
+
+	bool bHaveStats = false;
+	UE::Zen::FZenProjectStats ZenStats;
+
+	if (TSharedPtr<UE::Zen::FZenServiceInstance> ServiceInstance = ZenServiceInstance.Get())
+	{
+		if (ServiceInstance->GetProjectStats(ZenStats))
+		{
+			bHaveStats = true;
+		}
+	}
+
+	int64 TotalHits = ZenStats.General.Op.HitCount + ZenStats.General.Chunk.HitCount;
+	int64 TotalMisses = ZenStats.General.Op.MissCount + ZenStats.General.Chunk.MissCount;
+	int64 TotalWrites = ZenStats.General.Op.WriteCount + ZenStats.General.Chunk.WriteCount;
+
+	const static FNumberFormattingOptions SingleDecimalFormatting = FNumberFormattingOptions()
+		.SetUseGrouping(true)
+		.SetMinimumFractionalDigits(1)
+		.SetMaximumFractionalDigits(1);
+
+	int32 Row = 0;
+
+	const float RowMargin = 0.0f;
+	const float TitleMargin = 10.0f;
+	const float ColumnMargin = 10.0f;
+	const FSlateColor TitleColor = FStyleColors::AccentWhite;
+	const FSlateFontInfo TitleFont = FCoreStyle::GetDefaultFontStyle("Bold", 10);
+
+
+	Panel->AddSlot(0, Row)
+	.ColumnSpan(2)
+	[
+		SNew(STextBlock)
+		.Margin(FMargin(ColumnMargin, RowMargin))
+		.ColorAndOpacity(TitleColor)
+		.Font(TitleFont)
+		.Text(LOCTEXT("CAS", "Local Project Store"))
+	];
+
+	Row++;
+
+	Panel->AddSlot(0, Row)
+	[
+		SNew(STextBlock)
+		.Margin(FMargin(ColumnMargin, RowMargin))
+		.Text(LOCTEXT("ProjectHitRate", "Hit Rate"))
+	];
+
+	Panel->AddSlot(1, Row)
+	[
+		SNew(STextBlock)
+		.Margin(FMargin(ColumnMargin, RowMargin))
+		.Text_Lambda([bHaveStats, TotalHits, TotalMisses]
+		{
+			if (bHaveStats && ((TotalHits+TotalMisses) > 0))
+			{
+				double HitRate = static_cast<double>(TotalHits) / static_cast<double>(TotalHits + TotalMisses);
+				return FText::AsPercent(HitRate, &SingleDecimalFormatting);
+			}
+			return LOCTEXT("UnavailableValue", "-");
+		})
+	];
+
+	Row++;
+
+	Panel->AddSlot(0, Row)
+	[
+		SNew(STextBlock)
+		.Margin(FMargin(ColumnMargin, RowMargin))
+		.Text(LOCTEXT("ProjectHitQuantity", "Hits"))
+	];
+
+	Panel->AddSlot(1, Row)
+	[
+		SNew(STextBlock)
+		.Margin(FMargin(ColumnMargin, RowMargin))
+		.Text_Lambda([bHaveStats, TotalHits]
+		{
+			if (bHaveStats)
+			{
+				return FText::AsNumber(TotalHits);
+			}
+			return LOCTEXT("UnavailableValue", "-");
+		})
+	];
+
+	Row++;
+
+	Panel->AddSlot(0, Row)
+	[
+		SNew(STextBlock)
+		.Margin(FMargin(ColumnMargin, RowMargin))
+		.Text(LOCTEXT("ProjectMissQuantity", "Misses"))
+	];
+
+	Panel->AddSlot(1, Row)
+	[
+		SNew(STextBlock)
+		.Margin(FMargin(ColumnMargin, RowMargin))
+		.Text_Lambda([bHaveStats, TotalMisses]
+		{
+			if (bHaveStats)
+			{
+				return FText::AsNumber(TotalMisses);
+			}
+			return LOCTEXT("UnavailableValue", "-");
+		})
+	];
+
+	Row++;
+
+	Panel->AddSlot(0, Row)
+	[
+		SNew(STextBlock)
+		.Margin(FMargin(ColumnMargin, RowMargin))
+		.Text(LOCTEXT("ProjectWriteQuantity", "Writes"))
+	];
+
+	Panel->AddSlot(1, Row)
+	[
+		SNew(STextBlock)
+		.Margin(FMargin(ColumnMargin, RowMargin))
+		.Text_Lambda([bHaveStats, TotalWrites]
+		{
+			if (bHaveStats)
+			{
+				return FText::AsNumber(TotalWrites);
+			}
+			return LOCTEXT("UnavailableValue", "-");
+		})
+	];
+
+	Row++;
+
+	Panel->AddSlot(0, Row)
+	[
+		SNew(STextBlock)
+		.Margin(FMargin(ColumnMargin, RowMargin))
+		.Text(LOCTEXT("ProjectRequests", "Requests"))
+	];
+
+	Panel->AddSlot(1, Row)
+	[
+		SNew(STextBlock)
+		.Margin(FMargin(ColumnMargin, RowMargin))
+		.Text_Lambda([bHaveStats, Requests = ZenStats.General.RequestCount]
+		{
+			if (bHaveStats)
+			{
+				return FText::AsNumber(Requests);
+			}
+			return LOCTEXT("UnavailableValue", "-");
+		})
+	];
+
+	Row++;
+
+	Panel->AddSlot(0, Row)
+		[
+			SNew(STextBlock)
+			.Margin(FMargin(ColumnMargin, RowMargin))
+		.Text(LOCTEXT("ProjectBadRequests", "Bad Requests"))
+		];
+
+	Panel->AddSlot(1, Row)
+		[
+			SNew(STextBlock)
+			.Margin(FMargin(ColumnMargin, RowMargin))
+		.Text_Lambda([bHaveStats, BadRequests = ZenStats.General.BadRequestCount]
+			{
+				if (bHaveStats)
+				{
+					return FText::AsNumber(BadRequests);
+				}
+	return LOCTEXT("UnavailableValue", "-");
+			})
+		];
+
+	Row++;
+
+	return Panel;
+}
+
+#undef LOCTEXT_NAMESPACE
diff --git a/Engine/Source/Editor/ZenWidgets/Private/SZenServiceStatus.cpp b/Engine/Source/Editor/ZenWidgets/Private/SZenServiceStatus.cpp
new file mode 100644
--- /dev/null
+++ b/Engine/Source/Editor/ZenWidgets/Private/SZenServiceStatus.cpp
@@ -0,0 +1,280 @@
+// Copyright Epic Games, Inc. All Rights Reserved.
+
+#include "SZenServiceStatus.h"
+#include "Experimental/ZenServerInterface.h"
+#include "Internationalization/FastDecimalFormat.h"
+#include "Math/BasicMathExpressionEvaluator.h"
+#include "Math/UnitConversion.h"
+#include "Misc/ExpressionParser.h"
+#include "Styling/StyleColors.h"
+#include "Widgets/Images/SImage.h"
+#include "Widgets/Input/SButton.h"
+#include "Widgets/Layout/SGridPanel.h"
+#include "Widgets/SBoxPanel.h"
+#include "Widgets/Text/STextBlock.h"
+
+#define LOCTEXT_NAMESPACE "ZenDashboard"
+
+void SZenServiceStatus::Construct(const FArguments& InArgs)
+{
+	ZenServiceInstance = InArgs._ZenServiceInstance;
+
+	this->ChildSlot
+	[
+		SNew(SVerticalBox)
+
+		+ SVerticalBox::Slot()
+		.AutoHeight()
+		.Padding(0, 0, 0, 0)
+		.Expose(GridSlot)
+		[
+			GetGridPanel()
+		]
+	];
+
+	UpdateState(0.0, 0.0);
+	RegisterActiveTimer(0.5f, FWidgetActiveTimerDelegate::CreateSP(this, &SZenServiceStatus::UpdateState));
+}
+
+EActiveTimerReturnType SZenServiceStatus::UpdateState(double InCurrentTime, float InDeltaTime)
+{
+	uint32 LastStateIndex = ActiveStateIndex.load(std::memory_order_relaxed);
+	uint32 NextStateIndex = (LastStateIndex + 1)%NumState;
+	FState& NextState = State[NextStateIndex];
+	NextState.bGotRunContext = UE::Zen::TryGetLocalServiceRunContext(NextState.RunContext);
+	NextState.Version = UE::Zen::GetLocalServiceInstallVersion(false);
+	if (NextState.bGotRunContext)
+	{
+		NextState.bIsRunning = UE::Zen::IsLocalServiceRunning(*NextState.RunContext.GetDataPath(), &NextState.LocalPort);
+	}
+
+	NextState.GCStatus = UE::Zen::FGCStatus{};
+	NextState.bHaveStats = false;
+	if (NextState.bIsRunning)
+	{
+		if (TSharedPtr<UE::Zen::FZenServiceInstance> ServiceInstance = ZenServiceInstance.Get())
+		{
+			ServiceInstance->GetGCStatus(NextState.GCStatus);
+			if (ServiceInstance->GetCacheStats(NextState.ZenCacheStats) && ServiceInstance->GetProjectStats(NextState.ZenProjectStats))
+			{
+				NextState.bHaveStats = true;
+			}
+		}
+	}
+
+	ActiveStateIndex.compare_exchange_weak(LastStateIndex, NextStateIndex, std::memory_order_relaxed);
+
+	return EActiveTimerReturnType::Continue;
+}
+
+FReply SZenServiceStatus::ExploreDataPath_OnClicked()
+{
+	const FState& CurrentState = GetCurrentState();
+	if (CurrentState.bGotRunContext)
+	{
+		FString FullPath(FPaths::ConvertRelativePathToFull(CurrentState.RunContext.GetDataPath()));
+		FPlatformProcess::ExploreFolder(*FullPath);
+		return FReply::Handled();
+	}
+	return FReply::Unhandled();
+}
+
+const SZenServiceStatus::FState& SZenServiceStatus::GetCurrentState() const
+{
+	return State[ActiveStateIndex.load(std::memory_order_relaxed)];
+}
+
+TSharedRef<SWidget> SZenServiceStatus::GetGridPanel()
+{
+	TSharedRef<SGridPanel> Panel = SNew(SGridPanel);
+	
+	const static FNumberFormattingOptions SingleDecimalFormatting = FNumberFormattingOptions()
+		.SetUseGrouping(true)
+		.SetMinimumFractionalDigits(1)
+		.SetMaximumFractionalDigits(1);
+
+	int32 Row = 0;
+
+	const float RowMargin = 0.0f;
+	const float ColumnMargin = 10.0f;
+	const FSlateColor TitleColor = FStyleColors::AccentWhite;
+	const FSlateFontInfo TitleFont = FCoreStyle::GetDefaultFontStyle("Bold", 10);
+
+	Panel->AddSlot(0, Row)
+	[
+		SNew(STextBlock)
+		.Margin(FMargin(ColumnMargin, RowMargin))
+		.ColorAndOpacity(TitleColor)
+		.Font(TitleFont)
+		.Text(LOCTEXT("ServiceStatus_Status", "Status"))
+	];
+
+	Panel->AddSlot(1, Row)
+	[
+		SNew(STextBlock)
+		.Margin(FMargin(ColumnMargin, RowMargin))
+		.Text_Lambda([this]
+		{
+			const FState& CurrentState = GetCurrentState();
+			FText ServiceStatus = LOCTEXT("ServiceStatus_NotInstalled", "Not installed");
+			if (CurrentState.bGotRunContext)
+			{
+				if (CurrentState.bIsRunning && (CurrentState.LocalPort != 0))
+				{
+					ServiceStatus = LOCTEXT("ServiceStatus_Running", "Running");
+				}
+				else
+				{
+					ServiceStatus = LOCTEXT("ServiceStatus_Stopped", "Stopped");
+				}
+			}
+			return ServiceStatus;
+		})
+	];
+
+	++Row;
+
+	Panel->AddSlot(0, Row)
+	[
+		SNew(STextBlock)
+		.Margin(FMargin(ColumnMargin, RowMargin))
+		.ColorAndOpacity(TitleColor)
+		.Font(TitleFont)
+		.Text(LOCTEXT("ServiceStatus_Version", "Version"))
+	];
+
+	Panel->AddSlot(1, Row)
+	[
+		SNew(STextBlock)
+		.Margin(FMargin(ColumnMargin, RowMargin))
+		.Text_Lambda([this]
+		{
+			const FState& CurrentState = GetCurrentState();
+			return FText::FromString(CurrentState.Version);
+		})
+	];
+
+	++Row;
+
+	Panel->AddSlot(0, Row)
+	[
+		SNew(STextBlock)
+		.Margin(FMargin(ColumnMargin, RowMargin))
+		.ColorAndOpacity(TitleColor)
+		.Font(TitleFont)
+		.Text(LOCTEXT("ServiceStatus_Port", "Port"))
+	];
+
+	Panel->AddSlot(1, Row)
+	[
+		SNew(STextBlock)
+		.Margin(FMargin(ColumnMargin, RowMargin))
+		.Text_Lambda([this]
+		{
+			const FState& CurrentState = GetCurrentState();
+			return CurrentState.LocalPort == 0 ? LOCTEXT("ServiceStatus_NoPortValue", "-") : FText::AsNumber(CurrentState.LocalPort, &FNumberFormattingOptions::DefaultNoGrouping());
+		})
+	];
+
+	++Row;
+
+	Panel->AddSlot(0, Row)
+	[
+		SNew(STextBlock)
+		.Margin(FMargin(ColumnMargin, RowMargin))
+		.ColorAndOpacity(TitleColor)
+		.Font(TitleFont)
+		.Text(LOCTEXT("ServiceStatus_GCStatus", "GC Status"))
+	];
+
+	Panel->AddSlot(1, Row)
+	[
+		SNew(STextBlock)
+		.Margin(FMargin(ColumnMargin, RowMargin))
+		.Text_Lambda([this]
+		{
+			const FState& CurrentState = GetCurrentState();
+			return CurrentState.LocalPort == 0 ? LOCTEXT("ServiceStatus_NoGCStatusValue", "-") : FText::FromString(CurrentState.GCStatus.Description);
+		})
+	];
+
+	++Row;
+
+	Panel->AddSlot(0, Row)
+	[
+		SNew(STextBlock)
+		.Margin(FMargin(ColumnMargin, RowMargin))
+		.ColorAndOpacity(TitleColor)
+		.Font(TitleFont)
+		.Text(LOCTEXT("ServiceStatus_DiskSpace", "Disk Space"))
+	];
+
+	Panel->AddSlot(1, Row)
+	[
+		SNew(STextBlock)
+		.Margin(FMargin(ColumnMargin, RowMargin))
+		.Text_Lambda([this]
+		{
+			const FState& CurrentState = GetCurrentState();
+			const uint64 TotalDiskSpace = uint64(CurrentState.ZenCacheStats.General.Size.Disk) + uint64(CurrentState.ZenProjectStats.General.Size.Disk) + uint64(CurrentState.ZenCacheStats.CID.Size.Total);
+			return CurrentState.bHaveStats ? FText::AsMemory(TotalDiskSpace, (TotalDiskSpace > 1024) ? &SingleDecimalFormatting : nullptr, nullptr, EMemoryUnitStandard::IEC) : LOCTEXT("UnavailableValue", "-");
+		})
+	];
+
+	++Row;
+
+	Panel->AddSlot(0, Row)
+	.VAlign(VAlign_Bottom)
+	[
+		SNew(STextBlock)
+		.Margin(FMargin(ColumnMargin, RowMargin))
+		.ColorAndOpacity(TitleColor)
+		.Font(TitleFont)
+		.Text(LOCTEXT("ServiceStatus_DataPath", "Data path"))
+	];
+
+	Panel->AddSlot(1, Row)
+	[
+
+		SNew(SHorizontalBox)
+
+		+ SHorizontalBox::Slot()
+		.VAlign(VAlign_Bottom)
+		[
+			SNew(STextBlock)
+			.Margin(FMargin(ColumnMargin, RowMargin))
+			.Text_Lambda([this]
+			{
+				const FState& CurrentState = GetCurrentState();
+				return CurrentState.bGotRunContext ? FText::FromString(CurrentState.RunContext.GetDataPath()) : LOCTEXT("ServiceStatus_NoDataPathValue", "-");
+			})
+		]
+
+		+ SHorizontalBox::Slot()
+		.AutoWidth()
+		.VAlign(VAlign_Bottom)
+		[
+			SNew(SButton)
+			.ButtonStyle(&FAppStyle::Get().GetWidgetStyle<FButtonStyle>("SimpleButton"))
+			.ToolTipText(LOCTEXT("ExploreDataPathButtonToolTip", "Explore the Data Path"))
+			.IsEnabled_Lambda([this]
+			{
+				const FState& CurrentState = GetCurrentState();
+				return CurrentState.bGotRunContext;
+			})
+			.OnClicked(this, &SZenServiceStatus::ExploreDataPath_OnClicked)
+			[
+				SNew(SImage)
+				.Image(FAppStyle::Get().GetBrush("Zen.Icons.FolderExplore"))
+				.ColorAndOpacity(FSlateColor::UseForeground())
+			]
+		]
+	];
+
+	++Row;
+
+
+	return Panel;
+}
+
+#undef LOCTEXT_NAMESPACE
diff --git a/Engine/Source/Editor/ZenWidgets/Private/ZenServiceInstanceManager.cpp b/Engine/Source/Editor/ZenWidgets/Private/ZenServiceInstanceManager.cpp
new file mode 100644
--- /dev/null
+++ b/Engine/Source/Editor/ZenWidgets/Private/ZenServiceInstanceManager.cpp
@@ -0,0 +1,35 @@
+// Copyright Epic Games, Inc. All Rights Reserved.
+
+#include "ZenServiceInstanceManager.h"
+#include "Experimental/ZenServerInterface.h"
+
+
+namespace UE::Zen
+{
+
+TSharedPtr<FZenServiceInstance> FServiceInstanceManager::GetZenServiceInstance() const
+{
+	UE::Zen::FZenLocalServiceRunContext RunContext;
+	if (UE::Zen::TryGetLocalServiceRunContext(RunContext))
+	{
+		uint16 LocalPort = 0;
+		if (UE::Zen::IsLocalServiceRunning(*RunContext.GetDataPath(), &LocalPort) && (LocalPort != 0))
+		{
+			if (LocalPort != CurrentPort)
+			{
+				CurrentPort = LocalPort;
+				UE::Zen::FServiceSettings ServiceSettings;
+				ServiceSettings.SettingsVariant.Emplace<UE::Zen::FServiceConnectSettings>();
+				UE::Zen::FServiceConnectSettings& ConnectExistingSettings = ServiceSettings.SettingsVariant.Get<UE::Zen::FServiceConnectSettings>();
+				ConnectExistingSettings.HostName = TEXT("localhost");
+				ConnectExistingSettings.Port = LocalPort;
+				CurrentInstance = MakeShared<FZenServiceInstance>(MoveTemp(ServiceSettings));
+			}
+
+			return CurrentInstance;
+		}
+	}
+	return nullptr;
+}
+
+} // namespace UE::Zen
diff --git a/Engine/Source/Editor/ZenWidgets/Public/SMessageDialog.h b/Engine/Source/Editor/ZenWidgets/Public/SMessageDialog.h
new file mode 100644
--- /dev/null
+++ b/Engine/Source/Editor/ZenWidgets/Public/SMessageDialog.h
@@ -0,0 +1,16 @@
+// Copyright Epic Games, Inc. All Rights Reserved.
+
+#pragma once
+
+#include "Delegates/DelegateCombinations.h"
+#include "GenericPlatform/GenericPlatformMisc.h"
+#include "Internationalization/Text.h"
+
+class SWindow;
+class SWidget;
+
+// Largely a copy of UnrealEd Dialogs.h
+
+DECLARE_DELEGATE_TwoParams(FOnMsgDlgResult, const TSharedRef<SWindow>&, EAppReturnType::Type);
+
+EAppReturnType::Type OpenModalMessageDialog_Internal(EAppMsgCategory InMessageCategory, EAppMsgType::Type InMessageType, const FText& InMessage, const FText& InTitle, const TSharedPtr<const SWidget>& ModalParent);
diff --git a/Engine/Source/Editor/ZenWidgets/Public/SZenCacheStatistics.h b/Engine/Source/Editor/ZenWidgets/Public/SZenCacheStatistics.h
new file mode 100644
--- /dev/null
+++ b/Engine/Source/Editor/ZenWidgets/Public/SZenCacheStatistics.h
@@ -0,0 +1,33 @@
+// Copyright Epic Games, Inc. All Rights Reserved.
+
+#pragma once
+
+#include "Experimental/ZenServerInterface.h"
+#include "ZenServiceInstanceManager.h"
+#include "Templates/SharedPointer.h"
+#include "Widgets/DeclarativeSyntaxSupport.h"
+#include "Widgets/SBoxPanel.h"
+#include "Widgets/SCompoundWidget.h"
+#include "Widgets/SWidget.h"
+
+class ZENWIDGETS_API SZenCacheStatistics : public SCompoundWidget
+{
+	SLATE_BEGIN_ARGS(SZenCacheStatistics)
+		: _ZenServiceInstance(nullptr)
+	{ }
+
+	SLATE_ATTRIBUTE(TSharedPtr<UE::Zen::FZenServiceInstance>, ZenServiceInstance);
+
+	SLATE_END_ARGS()
+
+	void Construct(const FArguments& InArgs);
+
+private:
+
+	TSharedRef<SWidget> GetGridPanel();
+
+	EActiveTimerReturnType UpdateGridPanels(double InCurrentTime, float InDeltaTime);
+
+	SVerticalBox::FSlot* GridSlot = nullptr;
+	TAttribute<TSharedPtr<UE::Zen::FZenServiceInstance>> ZenServiceInstance;
+};
\ No newline at end of file
diff --git a/Engine/Source/Editor/ZenWidgets/Public/SZenCidStoreStatistics.h b/Engine/Source/Editor/ZenWidgets/Public/SZenCidStoreStatistics.h
new file mode 100644
--- /dev/null
+++ b/Engine/Source/Editor/ZenWidgets/Public/SZenCidStoreStatistics.h
@@ -0,0 +1,33 @@
+// Copyright Epic Games, Inc. All Rights Reserved.
+
+#pragma once
+
+#include "Experimental/ZenServerInterface.h"
+#include "ZenServiceInstanceManager.h"
+#include "Templates/SharedPointer.h"
+#include "Widgets/DeclarativeSyntaxSupport.h"
+#include "Widgets/SBoxPanel.h"
+#include "Widgets/SCompoundWidget.h"
+#include "Widgets/SWidget.h"
+
+class ZENWIDGETS_API SZenCidStoreStatistics : public SCompoundWidget
+{
+	SLATE_BEGIN_ARGS(SZenCidStoreStatistics)
+		: _ZenServiceInstance(nullptr)
+	{ }
+
+	SLATE_ATTRIBUTE(TSharedPtr<UE::Zen::FZenServiceInstance>, ZenServiceInstance);
+
+	SLATE_END_ARGS()
+
+	void Construct(const FArguments& InArgs);
+
+private:
+
+	TSharedRef<SWidget> GetGridPanel();
+
+	EActiveTimerReturnType UpdateGridPanels(double InCurrentTime, float InDeltaTime);
+
+	SVerticalBox::FSlot* GridSlot = nullptr;
+	TAttribute<TSharedPtr<UE::Zen::FZenServiceInstance>> ZenServiceInstance;
+};
\ No newline at end of file
diff --git a/Engine/Source/Editor/ZenWidgets/Public/SZenDialogs.h b/Engine/Source/Editor/ZenWidgets/Public/SZenDialogs.h
new file mode 100644
--- /dev/null
+++ b/Engine/Source/Editor/ZenWidgets/Public/SZenDialogs.h
@@ -0,0 +1,21 @@
+// Copyright Epic Games, Inc. All Rights Reserved.
+
+#pragma once
+
+#include "Templates/SharedPointer.h"
+#include "Types/SlateEnums.h"
+#include "Widgets/DeclarativeSyntaxSupport.h"
+#include "Widgets/SBoxPanel.h"
+#include "Widgets/SCompoundWidget.h"
+
+class SWidget;
+
+class ZENWIDGETS_API SZenStoreStausDialog : public SCompoundWidget
+{
+	SLATE_BEGIN_ARGS(SZenStoreStausDialog) {}
+	SLATE_END_ARGS()
+
+	void Construct(const FArguments& InArgs);
+
+private:
+};
diff --git a/Engine/Source/Editor/ZenWidgets/Public/SZenProjectStatistics.h b/Engine/Source/Editor/ZenWidgets/Public/SZenProjectStatistics.h
new file mode 100644
--- /dev/null
+++ b/Engine/Source/Editor/ZenWidgets/Public/SZenProjectStatistics.h
@@ -0,0 +1,33 @@
+// Copyright Epic Games, Inc. All Rights Reserved.
+
+#pragma once
+
+#include "Experimental/ZenServerInterface.h"
+#include "ZenServiceInstanceManager.h"
+#include "Templates/SharedPointer.h"
+#include "Widgets/DeclarativeSyntaxSupport.h"
+#include "Widgets/SBoxPanel.h"
+#include "Widgets/SCompoundWidget.h"
+#include "Widgets/SWidget.h"
+
+class ZENWIDGETS_API SZenProjectStatistics : public SCompoundWidget
+{
+	SLATE_BEGIN_ARGS(SZenProjectStatistics)
+		: _ZenServiceInstance(nullptr)
+	{ }
+
+	SLATE_ATTRIBUTE(TSharedPtr<UE::Zen::FZenServiceInstance>, ZenServiceInstance);
+
+	SLATE_END_ARGS()
+
+	void Construct(const FArguments& InArgs);
+
+private:
+
+	TSharedRef<SWidget> GetGridPanel();
+
+	EActiveTimerReturnType UpdateGridPanels(double InCurrentTime, float InDeltaTime);
+
+	SVerticalBox::FSlot* GridSlot = nullptr;
+	TAttribute<TSharedPtr<UE::Zen::FZenServiceInstance>> ZenServiceInstance;
+};
\ No newline at end of file
diff --git a/Engine/Source/Editor/ZenWidgets/Public/SZenServiceStatus.h b/Engine/Source/Editor/ZenWidgets/Public/SZenServiceStatus.h
new file mode 100644
--- /dev/null
+++ b/Engine/Source/Editor/ZenWidgets/Public/SZenServiceStatus.h
@@ -0,0 +1,51 @@
+// Copyright Epic Games, Inc. All Rights Reserved.
+
+#pragma once
+
+#include "Experimental/ZenServerInterface.h"
+#include "ZenServiceInstanceManager.h"
+#include "Templates/SharedPointer.h"
+#include "Widgets/DeclarativeSyntaxSupport.h"
+#include "Widgets/SBoxPanel.h"
+#include "Widgets/SCompoundWidget.h"
+#include "Widgets/SWidget.h"
+
+class ZENWIDGETS_API SZenServiceStatus : public SCompoundWidget
+{
+	SLATE_BEGIN_ARGS(SZenServiceStatus)
+		: _ZenServiceInstance(nullptr)
+	{ }
+
+	SLATE_ATTRIBUTE(TSharedPtr<UE::Zen::FZenServiceInstance>, ZenServiceInstance);
+
+	SLATE_END_ARGS()
+
+	void Construct(const FArguments& InArgs);
+
+private:
+	struct FState
+	{
+		UE::Zen::FZenLocalServiceRunContext RunContext;
+		FString Version;
+		UE::Zen::FGCStatus GCStatus;
+		UE::Zen::FZenCacheStats ZenCacheStats;
+		UE::Zen::FZenProjectStats ZenProjectStats;
+		uint16 LocalPort = 0;
+		bool bHaveStats = false;
+		bool bGotRunContext = false;
+		bool bIsRunning = false;
+	};
+	static constexpr uint32 NumState = 2;
+	FState State[NumState];
+	std::atomic<uint32> ActiveStateIndex = 0;
+
+
+	const FState& GetCurrentState() const;
+	TSharedRef<SWidget> GetGridPanel();
+	FReply ExploreDataPath_OnClicked();
+
+	EActiveTimerReturnType UpdateState(double InCurrentTime, float InDeltaTime);
+
+	SVerticalBox::FSlot* GridSlot = nullptr;
+	TAttribute<TSharedPtr<UE::Zen::FZenServiceInstance>> ZenServiceInstance;
+};
\ No newline at end of file
diff --git a/Engine/Source/Editor/ZenWidgets/Public/ZenServiceInstanceManager.h b/Engine/Source/Editor/ZenWidgets/Public/ZenServiceInstanceManager.h
new file mode 100644
--- /dev/null
+++ b/Engine/Source/Editor/ZenWidgets/Public/ZenServiceInstanceManager.h
@@ -0,0 +1,22 @@
+// Copyright Epic Games, Inc. All Rights Reserved.
+
+#pragma once
+
+#include "Templates/SharedPointer.h"
+
+namespace UE::Zen { class FZenServiceInstance; }
+
+namespace UE::Zen
+{
+
+class ZENWIDGETS_API FServiceInstanceManager
+{
+public:
+	TSharedPtr<FZenServiceInstance> GetZenServiceInstance() const;
+private:
+	mutable uint16 CurrentPort = 0;
+	mutable TSharedPtr<FZenServiceInstance> CurrentInstance;
+};
+
+} // namespace UE::Zen
+
diff --git a/Engine/Source/Editor/ZenWidgets/ZenWidgets.Build.cs b/Engine/Source/Editor/ZenWidgets/ZenWidgets.Build.cs
new file mode 100644
--- /dev/null
+++ b/Engine/Source/Editor/ZenWidgets/ZenWidgets.Build.cs
@@ -0,0 +1,26 @@
+// Copyright Epic Games, Inc. All Rights Reserved.
+
+using UnrealBuildTool;
+
+public class ZenWidgets : ModuleRules
+{
+	public ZenWidgets(ReadOnlyTargetRules Target)
+		 : base(Target)
+	{
+		PrivateDependencyModuleNames.AddRange(
+			new string[] {
+				"ApplicationCore",
+				"Core",
+				"CoreUObject",
+				"InputCore",
+				"Json",
+				"Projects",	
+				"Slate",
+				"SlateCore",
+				"StandaloneRenderer",
+				"Zen",
+			});
+
+		CppCompileWarningSettings.UnsafeTypeCastWarningLevel = WarningLevel.Error;
+	}
+}
diff --git a/Engine/Source/Programs/ZenDashboard/Private/SMessageDialog.cpp b/Engine/Source/Programs/ZenDashboard/Private/SMessageDialog.cpp
deleted file mode 100644
diff --git a/Engine/Source/Programs/ZenDashboard/Private/SMessageDialog.h b/Engine/Source/Programs/ZenDashboard/Private/SMessageDialog.h
deleted file mode 100644
diff --git a/Engine/Source/Programs/ZenDashboard/Private/SZenCacheStatistics.cpp b/Engine/Source/Programs/ZenDashboard/Private/SZenCacheStatistics.cpp
deleted file mode 100644
diff --git a/Engine/Source/Programs/ZenDashboard/Private/SZenCacheStatistics.h b/Engine/Source/Programs/ZenDashboard/Private/SZenCacheStatistics.h
deleted file mode 100644
diff --git a/Engine/Source/Programs/ZenDashboard/Private/SZenCidStoreStatistics.cpp b/Engine/Source/Programs/ZenDashboard/Private/SZenCidStoreStatistics.cpp
deleted file mode 100644
diff --git a/Engine/Source/Programs/ZenDashboard/Private/SZenCidStoreStatistics.h b/Engine/Source/Programs/ZenDashboard/Private/SZenCidStoreStatistics.h
deleted file mode 100644
diff --git a/Engine/Source/Programs/ZenDashboard/Private/SZenProjectStatistics.cpp b/Engine/Source/Programs/ZenDashboard/Private/SZenProjectStatistics.cpp
deleted file mode 100644
diff --git a/Engine/Source/Programs/ZenDashboard/Private/SZenProjectStatistics.h b/Engine/Source/Programs/ZenDashboard/Private/SZenProjectStatistics.h
deleted file mode 100644
diff --git a/Engine/Source/Programs/ZenDashboard/Private/SZenServiceStatus.cpp b/Engine/Source/Programs/ZenDashboard/Private/SZenServiceStatus.cpp
deleted file mode 100644
diff --git a/Engine/Source/Programs/ZenDashboard/Private/SZenServiceStatus.h b/Engine/Source/Programs/ZenDashboard/Private/SZenServiceStatus.h
deleted file mode 100644
diff --git a/Engine/Source/Programs/ZenDashboard/Private/ServiceInstanceManager.cpp b/Engine/Source/Programs/ZenDashboard/Private/ServiceInstanceManager.cpp
deleted file mode 100644
diff --git a/Engine/Source/Programs/ZenDashboard/Private/ServiceInstanceManager.h b/Engine/Source/Programs/ZenDashboard/Private/ServiceInstanceManager.h
deleted file mode 100644
diff --git a/Engine/Source/Programs/ZenDashboard/Private/ZenDashboardStyle.cpp b/Engine/Source/Programs/ZenDashboard/Private/ZenDashboardStyle.cpp
--- a/Engine/Source/Programs/ZenDashboard/Private/ZenDashboardStyle.cpp
+++ b/Engine/Source/Programs/ZenDashboard/Private/ZenDashboardStyle.cpp
@@ -56,8 +56,8 @@ TSharedRef< FSlateStyleSet > FZenDashboardStyle::Create()
 
 	Style.Set( "AppIcon", new IMAGE_BRUSH( "Icons/DefaultAppIcon", FVector2D(20,20)) );
 
-	Style.Set("Icons.FolderExplore", new EDITOR_IMAGE_BRUSH_SVG("Starship/Common/ContentBrowser", Icon12x12));
-	Style.Set("Icons.Clipboard", new EDITOR_IMAGE_BRUSH_SVG("Starship/Common/Clipboard", Icon16x16));
+	Style.Set("Zen.Icons.FolderExplore", new EDITOR_IMAGE_BRUSH_SVG("Starship/Common/ContentBrowser", Icon12x12));
+	Style.Set("Zen.Icons.Clipboard", new EDITOR_IMAGE_BRUSH_SVG("Starship/Common/Clipboard", Icon16x16));
 
 	const FTextBlockStyle DefaultText = FTextBlockStyle()
 		.SetFont(DEFAULT_FONT("Bold", 10))
diff --git a/Engine/Source/Programs/ZenDashboard/ZenDashboard.Build.cs b/Engine/Source/Programs/ZenDashboard/ZenDashboard.Build.cs
--- a/Engine/Source/Programs/ZenDashboard/ZenDashboard.Build.cs
+++ b/Engine/Source/Programs/ZenDashboard/ZenDashboard.Build.cs
@@ -16,8 +16,8 @@ public class ZenDashboard : ModuleRules
 				"Slate",
 				"SlateCore",
 				"StandaloneRenderer",
-				"Zen"
-				
+				"Zen",
+				"ZenWidgets"
 			});
 
 		if (Target.IsInPlatformGroup(UnrealPlatformGroup.Linux))
