Changed declartation of SingleDecimalFormat() to static

#Synced-CL 51538927

diff --git a/Engine/Source/Editor/DerivedDataWidgets/Private/SDerivedDataDialogs.cpp b/Engine/Source/Editor/DerivedDataWidgets/Private/SDerivedDataDialogs.cpp
--- a/Engine/Source/Editor/DerivedDataWidgets/Private/SDerivedDataDialogs.cpp
+++ b/Engine/Source/Editor/DerivedDataWidgets/Private/SDerivedDataDialogs.cpp
@@ -29,7 +29,7 @@ static FString ZeroDecimalFormat(ValueType Value)
 	return FastDecimalFormat::NumberToString(Value, ExpressionParser::GetLocalizedNumberFormattingRules(), NumberFormattingOptions);
 }
 
-FString SingleDecimalFormat(double Value)
+static FString SingleDecimalFormat(double Value)
 {
 	const FNumberFormattingOptions NumberFormattingOptions = FNumberFormattingOptions()
 		.SetUseGrouping(true)
diff --git a/Engine/Source/Editor/VirtualizationEditor/Private/SVirtualAssetsStatistics.cpp b/Engine/Source/Editor/VirtualizationEditor/Private/SVirtualAssetsStatistics.cpp
--- a/Engine/Source/Editor/VirtualizationEditor/Private/SVirtualAssetsStatistics.cpp
+++ b/Engine/Source/Editor/VirtualizationEditor/Private/SVirtualAssetsStatistics.cpp
@@ -46,7 +46,7 @@ class SWidget;
 namespace
 {
 
-FString SingleDecimalFormat(double Value)
+static FString SingleDecimalFormat(double Value)
 {
 	const FNumberFormattingOptions NumberFormattingOptions = FNumberFormattingOptions()
 		.SetUseGrouping(true)
